PreLoadScreenManager will now cap possible tick time to 5 seconds. Prevents breakpoints or backgrounding from having unreasonably huge tick swings.
#tests Android Client

diff --git a/Engine/Source/Runtime/PreLoadScreen/Private/PreLoadScreenManager.cpp b/Engine/Source/Runtime/PreLoadScreen/Private/PreLoadScreenManager.cpp
--- a/Engine/Source/Runtime/PreLoadScreen/Private/PreLoadScreenManager.cpp
+++ b/Engine/Source/Runtime/PreLoadScreen/Private/PreLoadScreenManager.cpp
@@ -291,6 +291,11 @@ void FPreLoadScreenManager::GameLogicFrameTick()
         double DeltaTime = CurrentTime - LastTickTime;
         LastTickTime = CurrentTime;
 
+		//Clamp to what should be more then any max reasonable time. This is to help with cases of
+		//backgrounding or setting breakpoints to trigger huge ticks
+		const double MaxTickTime = 5.0;
+		DeltaTime = FMath::Min(DeltaTime, MaxTickTime);
+
         //We have to manually tick everything as we are looping the main thread here
         FTicker::GetCoreTicker().Tick(DeltaTime);
         FThreadManager::Get().Tick();
