Fixing bug where IOS was not sending background/'foreground delegate calls during a PreLoadScreen.
#tests Android Client, IOS Client

diff --git a/Engine/Source/Runtime/ApplicationCore/Private/IOS/IOSAppDelegate.cpp b/Engine/Source/Runtime/ApplicationCore/Private/IOS/IOSAppDelegate.cpp
--- a/Engine/Source/Runtime/ApplicationCore/Private/IOS/IOSAppDelegate.cpp
+++ b/Engine/Source/Runtime/ApplicationCore/Private/IOS/IOSAppDelegate.cpp
@@ -1376,6 +1376,8 @@ FCriticalSection RenderSuspend;
 	 If your application supports background execution, this method is called
 	 instead of applicationWillTerminate: when the user quits.
 	 */
+
+#if BUILD_EMBEDDED_APP
 	FIOSAsyncTask* AsyncTask = [[FIOSAsyncTask alloc] init];
 	AsyncTask.GameThreadCallback = ^ bool(void)
 	{
@@ -1383,6 +1385,9 @@ FCriticalSection RenderSuspend;
 		return true;
 	};
 	[AsyncTask FinishedTask];
+#else
+	FCoreDelegates::ApplicationWillEnterBackgroundDelegate.Broadcast();
+#endif //BUILD_EBMEDDED_APP
 }
 
 - (void)applicationWillEnterForeground:(UIApplication *)application
diff --git a/Engine/Source/Runtime/PreLoadScreen/Private/PreLoadScreenManager.cpp b/Engine/Source/Runtime/PreLoadScreen/Private/PreLoadScreenManager.cpp
--- a/Engine/Source/Runtime/PreLoadScreen/Private/PreLoadScreenManager.cpp
+++ b/Engine/Source/Runtime/PreLoadScreen/Private/PreLoadScreenManager.cpp
@@ -300,13 +300,9 @@ void FPreLoadScreenManager::GameLogicFrameTick()
         FTicker::GetCoreTicker().Tick(DeltaTime);
         FThreadManager::Get().Tick();
 
-        //Tick android events
-#if PLATFORM_ANDROID
-#if USE_ANDROID_EVENTS
-    // Process any Android events or we may have issues returning from background
-        FAppEventManager::GetInstance()->Tick();
-#endif
-#endif
+		//Tick any platform specific things we need here
+		PlatformSpecificGameLogicFrameTick();
+
         //Tick the Active Screen
         ActivePreLoadScreen->Tick(DeltaTime);
 
@@ -322,6 +318,17 @@ void FPreLoadScreenManager::GameLogicFrameTick()
     }
 }
 
+void FPreLoadScreenManager::PlatformSpecificGameLogicFrameTick()
+{
+#if PLATFORM_ANDROID
+	Android_PlatformSpecificGameLogicFrameTick();
+#endif //PLATFORM_ANDROID
+
+#if PLATFORM_IOS
+	IOS_PlatformSpecificGameLogicFrameTick();
+#endif //PLATFORM_IOS
+}
+
 bool FPreLoadScreenManager::ShouldRender()
 {
     FScopeLock ScopeLock(&RenderingEnabledCriticalSection);
@@ -500,3 +507,21 @@ void FPreLoadScreenManager::CleanUpResources()
     //other game specific settings containers should be cleaned up by their screens/modules
     BeginCleanup(&FPreLoadSettingsContainerBase::Get());
 }
+
+#if PLATFORM_ANDROID
+void FPreLoadScreenManager::Android_PlatformSpecificGameLogicFrameTick()
+{
+#if USE_ANDROID_EVENTS
+	// Process any Android events or we may have issues returning from background
+	FAppEventManager::GetInstance()->Tick();
+#endif //USE_ANDROIID_EVENTS
+}
+#endif //PLATFORM_ANDROID
+
+#if PLATFORM_IOS
+void FPreLoadScreenManager::IOS_PlatformSpecificGameLogicFrameTick()
+{
+	// drain the async task queue from the game thread
+	[FIOSAsyncTask ProcessAsyncTasks];
+}
+#endif
\ No newline at end of file
diff --git a/Engine/Source/Runtime/PreLoadScreen/Public/PreLoadScreenManager.h b/Engine/Source/Runtime/PreLoadScreen/Public/PreLoadScreenManager.h
--- a/Engine/Source/Runtime/PreLoadScreen/Public/PreLoadScreenManager.h
+++ b/Engine/Source/Runtime/PreLoadScreen/Public/PreLoadScreenManager.h
@@ -112,6 +112,8 @@ protected:
     void GameLogicFrameTick();
     void EarlyPlayRenderFrameTick();
 
+	void PlatformSpecificGameLogicFrameTick();
+
     /*** These functions describe how everything is handled during an non-Early PreLoadPlay. Everything is handled asynchronously in this case with a standalone renderer ***/
     void HandleEngineLoadingPlay();
 
@@ -144,4 +146,13 @@ protected:
 
     float OriginalSlateSleepVariableValue;
     bool bIsEngineLoadingComplete;
+
+private:
+#if PLATFORM_ANDROID
+	void Android_PlatformSpecificGameLogicFrameTick();
+#endif //PLATFORM_ANDROID
+
+#if PLATFORM_IOS
+	void IOS_PlatformSpecificGameLogicFrameTick();
+#endif //PLATFORM_IOS
 };
