Added ability to use system fonts for a PreLoadScreenSettingsContainer font. Allows you to lower initial size by not including fonts.
Added ability for PreLoadScreenSettingsContainer to load slate brushes in groups. This way not all brushes are loaded at the same time.
#tests Windows Editor, iOS client

diff --git a/Engine/Source/Runtime/PreLoadScreen/Private/PreLoadScreenBase.cpp b/Engine/Source/Runtime/PreLoadScreen/Private/PreLoadScreenBase.cpp
--- a/Engine/Source/Runtime/PreLoadScreen/Private/PreLoadScreenBase.cpp
+++ b/Engine/Source/Runtime/PreLoadScreen/Private/PreLoadScreenBase.cpp
@@ -39,6 +39,11 @@ void FPreLoadScreenBase::InitSettingsFromConfig(const FString& ConfigFileName)
         Config->GetFloat(*UIConfigSection, TEXT("TimeToDisplayEachBackground"), TimeToDisplayEachBackground);
         FPreLoadSettingsContainerBase::Get().TimeToDisplayEachBackground = TimeToDisplayEachBackground;
 
+		//Parse LoadingGroups. You want to do this before ScreenGroupings and CustomImageBrushes
+		TArray<FString> LoadingGroups;
+		Config->GetArray(*UIConfigSection, TEXT("LoadingGroups"), LoadingGroups);
+		SettingsContainer->ParseLoadingGroups(LoadingGroups);
+
         //Parse custom brushes
         TArray<FString> CustomImageBrushes;
         Config->GetArray(*UIConfigSection, TEXT("CustomImageBrushes"), CustomImageBrushes);
@@ -71,10 +76,16 @@ void FPreLoadScreenBase::InitSettingsFromConfig(const FString& ConfigFileName)
             SettingsContainer->ParseFontConfigEntry(FontEntry);
         }
 
-        //Parse Screen Grouping order
-        TArray<FString> ScreenGroupingDisplayOrder;
-        Config->GetArray(*UIConfigSection, TEXT("ScreenDisplayOrder"), ScreenGroupingDisplayOrder);
-        SettingsContainer->ScreenDisplayOrder = ScreenGroupingDisplayOrder;
+		TArray<FString> ScreenOrders;
+		Config->GetArray(*UIConfigSection, TEXT("ScreenOrders"), ScreenOrders);
+		//Support old format of ScreenDisplayOrder=
+		if (ScreenOrders.Num() == 0)
+		{
+			Config->GetArray(*UIConfigSection, TEXT("ScreenDisplayOrder"), ScreenOrders);
+		}
+		SettingsContainer->ParseAllScreenOrderEntries(LoadingGroups, ScreenOrders);
+
+		SettingsContainer->PerformInitialAssetLoad();
     }
 }
 
diff --git a/Engine/Source/Runtime/PreLoadScreen/Private/PreLoadSettingsContainer.cpp b/Engine/Source/Runtime/PreLoadScreen/Private/PreLoadSettingsContainer.cpp
--- a/Engine/Source/Runtime/PreLoadScreen/Private/PreLoadSettingsContainer.cpp
+++ b/Engine/Source/Runtime/PreLoadScreen/Private/PreLoadSettingsContainer.cpp
@@ -5,51 +5,97 @@
 #include "Brushes/SlateDynamicImageBrush.h"
 #include "Fonts/CompositeFont.h"
 #include "Fonts/UnicodeBlockRange.h"
+#include "Misc/FileHelper.h"
 #include "Misc/Paths.h"
 #include "Framework/Application/SlateApplication.h"
 
 #include "Internationalization/Culture.h"
 
+FString FPreLoadSettingsContainerBase::UseSystemFontOverride = TEXT("SYSTEM");
+FString FPreLoadSettingsContainerBase::DefaultInitialLoadingGroupIdentifier = TEXT("InitialLoad");
 FPreLoadSettingsContainerBase* FPreLoadSettingsContainerBase::Instance;
 
 FPreLoadSettingsContainerBase::~FPreLoadSettingsContainerBase()
 {
-    for (auto KVPair : BrushResources)
-    {
-        FSlateApplication::Get().GetRenderer()->ReleaseDynamicResource(*KVPair.Value);
+	for (auto KVPair : BrushResources)
+	{
+		FSlateApplication::Get().GetRenderer()->ReleaseDynamicResource(*KVPair.Value);
 		delete KVPair.Value;
-    }
-    BrushResources.Empty();
-    LocalizedTextResources.Empty();
-    FontResources.Empty();
-    ScreenGroupings.Empty();
+	}
+	BrushResources.Empty();
+	LocalizedTextResources.Empty();
+	FontResources.Empty();
+	ScreenGroupings.Empty();
 
-    Instance = nullptr;
+	Instance = nullptr;
 }
 
 bool FPreLoadSettingsContainerBase::IsValidBrushConfig(TArray<FString>& SplitConfigEntry)
 {
-    return (SplitConfigEntry.Num() == 4);
+	return (SplitConfigEntry.Num() == 4);
 }
 
 void FPreLoadSettingsContainerBase::ParseBrushConfigEntry(const FString& ConfigEntry)
 {
-    TArray<FString> BrushComponents;
-    ConfigEntry.ParseIntoArray(BrushComponents, TEXT(","), true);
-    if (ensureAlwaysMsgf(IsValidBrushConfig(BrushComponents), TEXT("Invalid Custom Brush in config. Exptected Format: +CustomImageBrushes=(Identifier,Filename,Width,Height). Config Entry: %s"), *ConfigEntry))
-    {
-        //Clean up the identifier to remove extra spaces and the first (
-        FString Identifier = BrushComponents[0];
-        Identifier.TrimStartAndEndInline();
-        Identifier.RemoveFromStart("(");
-
-        FString FilePath = ConvertIfPluginRelativeContentPath(BrushComponents[1]);
-
-        float Width = FCString::Atof(*BrushComponents[2]);
-        float Height = FCString::Atof(*BrushComponents[3]);
-
-        CreateCustomSlateImageBrush(Identifier, FilePath, FVector2D(Width, Height));
-    }
+	if (ensureAlwaysMsgf(!ConfigEntry.IsEmpty(), TEXT("Attempt to parse empty ConfigEntry!")))
+	{
+		bool bWasValidEntry = true;
+
+		TArray<FString> BrushComponents;
+		ConfigEntry.ParseIntoArray(BrushComponents, TEXT("("), true);
+
+		TArray<FString> BasicBrushComponents;
+		BrushComponents[0].ParseIntoArray(BasicBrushComponents, TEXT(","), true);
+
+		TArray<FString> LoadingGroupIdentifiers;
+
+		if (BrushComponents.Num() == 1)
+		{
+			LoadingGroupIdentifiers.Add(DefaultInitialLoadingGroupIdentifier);
+		}
+		else if (BrushComponents.Num() == 2)
+		{
+			TArray<FString> LoadingIdentifiersForBrush;
+			BrushComponents[1].ParseIntoArray(LoadingIdentifiersForBrush, TEXT(","), true);
+
+			for (FString LoadingIdentifier : LoadingIdentifiersForBrush)
+			{
+				LoadingIdentifier.TrimStartAndEndInline();
+				LoadingIdentifier.RemoveFromStart("(");
+				LoadingIdentifier.RemoveFromEnd(")");
+				LoadingIdentifier.RemoveFromEnd(")");
+
+				LoadingGroupIdentifiers.Add(LoadingIdentifier);
+			}
+		}
+
+		//Flag if we didn't find any LoadingIdentifiers or if our basic brush information is bad
+		bWasValidEntry = bWasValidEntry && (LoadingGroupIdentifiers.Num() > 0) && IsValidBrushConfig(BasicBrushComponents);
+
+		if (ensureAlwaysMsgf(bWasValidEntry, TEXT("Invalid Custom Brush in config. Exptected Format: +CustomImageBrushes=(Identifier,Filename,Width,Height,LoadGroupIdentifier) or +CustomImageBrushes=(Identifier,Filename,Width,Height). Config Entry: %s"), *ConfigEntry))
+		{
+			//Clean up the identifier to remove extra spaces and the first (
+			FString Identifier = BasicBrushComponents[0];
+			Identifier.TrimStartAndEndInline();
+			Identifier.RemoveFromStart("(");
+
+			FString FilePath = ConvertIfPluginRelativeContentPath(BasicBrushComponents[1]);
+
+			float Width = FCString::Atof(*BasicBrushComponents[2]);
+			float Height = FCString::Atof(*BasicBrushComponents[3]);
+
+			FCustomBrushDefine NewBrushDefine(Identifier, FilePath, FVector2D(Width, Height));
+
+			for (const FString& LoadingGroupIdentifier : LoadingGroupIdentifiers)
+			{
+				FCustomBrushLoadingGroup* FoundBrushLoadingGroup = BrushLoadingGroups.Find(*LoadingGroupIdentifier);
+				if (ensureAlwaysMsgf((FoundBrushLoadingGroup != nullptr), TEXT("LoadingGroup not found for parsed LoadingGroupIdentifier %s. Config Entry: %s"), *LoadingGroupIdentifier, *ConfigEntry))
+				{
+					FoundBrushLoadingGroup->CustomBrushDefinesToLoad.Add(NewBrushDefine);
+				}
+			}
+		}
+	}
 }
 
 bool FPreLoadSettingsContainerBase::IsValidFontConfigString(TArray<FString>& SplitConfigEntry)
@@ -73,123 +119,256 @@ void FPreLoadSettingsContainerBase::ParseFontConfigEntry(const FString& SplitCon
         FString FilePath = FontComponents[2];
         FilePath.TrimStartAndEndInline();
         FilePath.RemoveFromEnd(TEXT(")"));
-        FilePath = ConvertIfPluginRelativeContentPath(FilePath);
+		
+		//Only convert pluging path if we aren't using the SystemFontOverride
+		if (FilePath != UseSystemFontOverride)
+		{
+			FilePath = ConvertIfPluginRelativeContentPath(FilePath);
+		}
 
         BuildCustomFont(Identifier, Language, FilePath);
     }
 }
 
+void FPreLoadSettingsContainerBase::ParseLoadingGroups(TArray<FString>& LoadingGroupIdentifiers)
+{
+	//if we have no loading groups, go ahead and add the default identifier
+	if (LoadingGroupIdentifiers.Num() == 0)
+	{
+		ScreenOrderByLoadingGroups.FindOrAdd(*DefaultInitialLoadingGroupIdentifier);
+		BrushLoadingGroups.FindOrAdd(*DefaultInitialLoadingGroupIdentifier);
+	}
+	else
+	{
+		for (const FString& LoadingGroupIdentifier : LoadingGroupIdentifiers)
+		{
+			ScreenOrderByLoadingGroups.FindOrAdd(*LoadingGroupIdentifier);
+			BrushLoadingGroups.FindOrAdd(*LoadingGroupIdentifier);
+		}
+	}
+}
+
+void FPreLoadSettingsContainerBase::ParseAllScreenOrderEntries(TArray<FString>& LoadingGroupsEntries, TArray<FString>& ScreenOrderEntries)
+{	
+	for (const FString& ScreenOrderEntry : ScreenOrderEntries)
+	{
+		ParseScreenOrderConfigString(ScreenOrderEntry);
+	}
+}
+
+void FPreLoadSettingsContainerBase::ParseScreenOrderConfigString(const FString& ScreenOrderEntry)
+{
+	TArray<FString> ScreenOrderComponents;
+	ScreenOrderEntry.ParseIntoArray(ScreenOrderComponents, TEXT(","), true);
+	
+	FString LoadingGroupIdentifier;
+	FString ScreenGroupingIdentifier;
+
+	//Parse assuming default identifier and we are only supplying the MarketScreen
+	if (ScreenOrderComponents.Num() == 1)
+	{
+		//We didn't supply an initial group identifier, so assume default
+		LoadingGroupIdentifier = DefaultInitialLoadingGroupIdentifier;
+		ScreenGroupingIdentifier = ScreenOrderComponents[0];
+	}
+	else if (ScreenOrderComponents.Num() == 2)
+	{
+		LoadingGroupIdentifier = ScreenOrderComponents[0];
+		ScreenGroupingIdentifier = ScreenOrderComponents[1];
+	}
+	else
+	{
+		ensureAlwaysMsgf(false, TEXT("Invalid ScreenOrder config entry! Expected format is either +ScreenOrders=(LoadingGroupsIdentifier,ScreenGroupingsIdentifier) or +ScreenOrders=(ScreenGroupingsIdentifier). Found: %s"), *ScreenOrderEntry);
+		return;
+	}
+
+	//Clean up entries to remove extra spaces and parenthesis
+	LoadingGroupIdentifier.RemoveSpacesInline();
+	LoadingGroupIdentifier.RemoveFromStart("(");
+	ScreenGroupingIdentifier.RemoveSpacesInline();
+	ScreenGroupingIdentifier.RemoveFromStart("(");
+	ScreenGroupingIdentifier.RemoveFromEnd(")");
+
+	if (ensureAlwaysMsgf((!LoadingGroupIdentifier.IsEmpty() && !ScreenGroupingIdentifier.IsEmpty()), TEXT("No valid Loading Group Identifier or ScreenGroupingIdentifier found for ScreenOrderEntry! %s"), *ScreenOrderEntry))
+	{
+		FScreenOrderByLoadingGroup* FoundLoadingGroup = ScreenOrderByLoadingGroups.Find(*LoadingGroupIdentifier);
+		FScreenGroupingBase* FoundScreenGrouping = ScreenGroupings.Find(*ScreenGroupingIdentifier);
+
+		if (ensureAlwaysMsgf((FoundLoadingGroup != nullptr), TEXT("Did not find LoadingGroup definition for LoadingGroup %s for entry  %s"), *LoadingGroupIdentifier, *ScreenOrderEntry))
+		{
+			if (ensureAlwaysMsgf((FoundScreenGrouping != nullptr), TEXT("Did not find ScreenGrouping definition for ScreenGrouping %s for entry  %s"), *ScreenGroupingIdentifier, *ScreenOrderEntry))
+			{
+				FoundLoadingGroup->ScreenGroupings.Add(*ScreenGroupingIdentifier);
+			}
+		}
+	}
+}
+
+void FPreLoadSettingsContainerBase::PerformInitialAssetLoad()
+{
+	//Try and load the default initial loading group, but only if we know we have it to avoid triggering any ensures
+	FCustomBrushLoadingGroup* FoundDefaultLoadingGroup = BrushLoadingGroups.Find(*DefaultInitialLoadingGroupIdentifier);
+	if (FoundDefaultLoadingGroup != nullptr)
+	{
+		LoadGrouping(*DefaultInitialLoadingGroupIdentifier);
+	}
+}
+
+void FPreLoadSettingsContainerBase::LoadGrouping(FName Identifier)
+{
+	FCustomBrushLoadingGroup* FoundBrushLoadingGroup = BrushLoadingGroups.Find(*Identifier.ToString());
+	if (ensureAlwaysMsgf((FoundBrushLoadingGroup != nullptr), TEXT("Could not find LoadGrouping for identifier:%s"), *Identifier.ToString()))
+	{
+		//only change if we can actually find the new LoadGrouping
+		CurrentLoadGroup = Identifier;
+
+		for (FCustomBrushDefine& BrushDefine : FoundBrushLoadingGroup->CustomBrushDefinesToLoad)
+		{
+			CreateCustomSlateImageBrush(BrushDefine.BrushIdentifier, BrushDefine.FilePath, BrushDefine.Size);
+		}
+	}
+}
+
 void FPreLoadSettingsContainerBase::BuildCustomFont(const FString& FontIdentifier, const FString& Language, const FString& FilePath)
 {
-    //Try and find existing font, if we can't, make a new one
-    TSharedPtr<FCompositeFont>& FontToBuild = FontResources.FindOrAdd(*FontIdentifier);
-    if (!FontToBuild.IsValid())
-    {
-        FontToBuild = MakeShared<FCompositeFont>();
-    }
+	if (FilePath == UseSystemFontOverride)
+	{
+		BuildSystemFont(FontIdentifier);
+	}
+	else
+	{
+		//Try and find existing font, if we can't, make a new one
+		TSharedPtr<FCompositeFont>& FontToBuild = FontResources.FindOrAdd(*FontIdentifier);
+		if (!FontToBuild.IsValid())
+		{
+			FontToBuild = MakeShared<FCompositeFont>();
+		}
+
+		if (ensureAlwaysMsgf(FontToBuild.IsValid(), TEXT("Error creating custom font!")))
+		{
+
+			//If En, then setup as default font
+			if (Language.Equals(TEXT("en")))
+			{
+				FontToBuild->DefaultTypeface.AppendFont(*FontIdentifier, FilePath, EFontHinting::Default, EFontLoadingPolicy::LazyLoad);
+			}
+			// if not en, we need to setup some block ranges and subfonts to handle special characters
+			else
+			{
+				TArray<EUnicodeBlockRange> BlockRangesForLanguage;
+				//Arabic
+				if (Language.Equals(TEXT("ar"), ESearchCase::IgnoreCase))
+				{
+					BlockRangesForLanguage.Add(EUnicodeBlockRange::Arabic);
+					BlockRangesForLanguage.Add(EUnicodeBlockRange::ArabicExtendedA);
+					BlockRangesForLanguage.Add(EUnicodeBlockRange::ArabicMathematicalAlphabeticSymbols);
+					BlockRangesForLanguage.Add(EUnicodeBlockRange::ArabicPresentationFormsA);
+					BlockRangesForLanguage.Add(EUnicodeBlockRange::ArabicPresentationFormsB);
+					BlockRangesForLanguage.Add(EUnicodeBlockRange::ArabicSupplement);
+				}
+				//Japanese
+				else if (Language.Equals(TEXT("ja"), ESearchCase::IgnoreCase))
+				{
+					BlockRangesForLanguage.Add(EUnicodeBlockRange::CJKCompatibility);
+					BlockRangesForLanguage.Add(EUnicodeBlockRange::CJKCompatibilityForms);
+					BlockRangesForLanguage.Add(EUnicodeBlockRange::CJKCompatibilityIdeographs);
+					BlockRangesForLanguage.Add(EUnicodeBlockRange::CJKCompatibilityIdeographsSupplement);
+					BlockRangesForLanguage.Add(EUnicodeBlockRange::CJKRadicalsSupplement);
+					BlockRangesForLanguage.Add(EUnicodeBlockRange::CJKStrokes);
+					BlockRangesForLanguage.Add(EUnicodeBlockRange::CJKSymbolsAndPunctuation);
+					BlockRangesForLanguage.Add(EUnicodeBlockRange::CJKUnifiedIdeographs);
+					BlockRangesForLanguage.Add(EUnicodeBlockRange::CJKUnifiedIdeographsExtensionA);
+					BlockRangesForLanguage.Add(EUnicodeBlockRange::CJKUnifiedIdeographsExtensionB);
+					BlockRangesForLanguage.Add(EUnicodeBlockRange::CJKUnifiedIdeographsExtensionC);
+					BlockRangesForLanguage.Add(EUnicodeBlockRange::CJKUnifiedIdeographsExtensionD);
+					BlockRangesForLanguage.Add(EUnicodeBlockRange::CJKUnifiedIdeographsExtensionE);
+					BlockRangesForLanguage.Add(EUnicodeBlockRange::EnclosedCJKLettersAndMonths);
+					BlockRangesForLanguage.Add(EUnicodeBlockRange::Hiragana);
+					BlockRangesForLanguage.Add(EUnicodeBlockRange::Katakana);
+					BlockRangesForLanguage.Add(EUnicodeBlockRange::KatakanaPhoneticExtensions);
+					BlockRangesForLanguage.Add(EUnicodeBlockRange::Kanbun);
+					BlockRangesForLanguage.Add(EUnicodeBlockRange::HalfwidthAndFullwidthForms);
+				}
+				//Korean
+				else if (Language.Equals(TEXT("ko"), ESearchCase::IgnoreCase))
+				{
+					BlockRangesForLanguage.Add(EUnicodeBlockRange::HangulJamo);
+					BlockRangesForLanguage.Add(EUnicodeBlockRange::HangulJamoExtendedA);
+					BlockRangesForLanguage.Add(EUnicodeBlockRange::HangulJamoExtendedB);
+					BlockRangesForLanguage.Add(EUnicodeBlockRange::HangulCompatibilityJamo);
+					BlockRangesForLanguage.Add(EUnicodeBlockRange::HangulSyllables);
+				}
+				//Simplified Chinese
+				else if (Language.Equals(TEXT("zh-hans"), ESearchCase::IgnoreCase))
+				{
+					BlockRangesForLanguage.Add(EUnicodeBlockRange::CJKCompatibility);
+					BlockRangesForLanguage.Add(EUnicodeBlockRange::CJKCompatibilityForms);
+					BlockRangesForLanguage.Add(EUnicodeBlockRange::CJKCompatibilityIdeographs);
+					BlockRangesForLanguage.Add(EUnicodeBlockRange::CJKCompatibilityIdeographsSupplement);
+					BlockRangesForLanguage.Add(EUnicodeBlockRange::CJKRadicalsSupplement);
+					BlockRangesForLanguage.Add(EUnicodeBlockRange::CJKStrokes);
+					BlockRangesForLanguage.Add(EUnicodeBlockRange::CJKSymbolsAndPunctuation);
+					BlockRangesForLanguage.Add(EUnicodeBlockRange::CJKUnifiedIdeographs);
+					BlockRangesForLanguage.Add(EUnicodeBlockRange::CJKUnifiedIdeographsExtensionA);
+					BlockRangesForLanguage.Add(EUnicodeBlockRange::CJKUnifiedIdeographsExtensionB);
+					BlockRangesForLanguage.Add(EUnicodeBlockRange::CJKUnifiedIdeographsExtensionC);
+					BlockRangesForLanguage.Add(EUnicodeBlockRange::CJKUnifiedIdeographsExtensionD);
+					BlockRangesForLanguage.Add(EUnicodeBlockRange::CJKUnifiedIdeographsExtensionE);
+					BlockRangesForLanguage.Add(EUnicodeBlockRange::EnclosedCJKLettersAndMonths);
+				}
+				//Traditional Chinese
+				else if (Language.Equals(TEXT("zh-hant"), ESearchCase::IgnoreCase))
+				{
+					BlockRangesForLanguage.Add(EUnicodeBlockRange::CJKCompatibility);
+					BlockRangesForLanguage.Add(EUnicodeBlockRange::CJKCompatibilityForms);
+					BlockRangesForLanguage.Add(EUnicodeBlockRange::CJKCompatibilityIdeographs);
+					BlockRangesForLanguage.Add(EUnicodeBlockRange::CJKCompatibilityIdeographsSupplement);
+					BlockRangesForLanguage.Add(EUnicodeBlockRange::CJKRadicalsSupplement);
+					BlockRangesForLanguage.Add(EUnicodeBlockRange::CJKStrokes);
+					BlockRangesForLanguage.Add(EUnicodeBlockRange::CJKSymbolsAndPunctuation);
+					BlockRangesForLanguage.Add(EUnicodeBlockRange::CJKUnifiedIdeographs);
+					BlockRangesForLanguage.Add(EUnicodeBlockRange::CJKUnifiedIdeographsExtensionA);
+					BlockRangesForLanguage.Add(EUnicodeBlockRange::CJKUnifiedIdeographsExtensionB);
+					BlockRangesForLanguage.Add(EUnicodeBlockRange::CJKUnifiedIdeographsExtensionC);
+					BlockRangesForLanguage.Add(EUnicodeBlockRange::CJKUnifiedIdeographsExtensionD);
+					BlockRangesForLanguage.Add(EUnicodeBlockRange::CJKUnifiedIdeographsExtensionE);
+					BlockRangesForLanguage.Add(EUnicodeBlockRange::EnclosedCJKLettersAndMonths);
+				}
+
+				//Build out actual sub font ranges
+				FCompositeSubFont& SubFont = FontToBuild->SubTypefaces[FontToBuild->SubTypefaces.AddDefaulted()];
+				SubFont.Cultures.Append(Language);
+				for (EUnicodeBlockRange& BlockRange : BlockRangesForLanguage)
+				{
+					SubFont.CharacterRanges.Add(FUnicodeBlockRange::GetUnicodeBlockRange(BlockRange).Range);
+				}
+
+				//Finally append actual font
+				SubFont.Typeface.AppendFont(*FontIdentifier, FilePath, EFontHinting::Default, EFontLoadingPolicy::LazyLoad);
+			}
+		}
+	}
+}
 
-    if (ensureAlwaysMsgf(FontToBuild.IsValid(), TEXT("Error creating custom font!")))
-    {
-        //If En, then setup as default font
-        if (Language.Equals(TEXT("en")))
-        {
-            FontToBuild->DefaultTypeface.AppendFont(*FontIdentifier, FilePath, EFontHinting::Default, EFontLoadingPolicy::LazyLoad);
-        }
-        // if not en, we need to setup some block ranges and subfonts to handle special characters
-        else
-        {
-            TArray<EUnicodeBlockRange> BlockRangesForLanguage;
-            //Arabic
-            if (Language.Equals(TEXT("ar"), ESearchCase::IgnoreCase))
-            {
-                BlockRangesForLanguage.Add(EUnicodeBlockRange::Arabic);
-                BlockRangesForLanguage.Add(EUnicodeBlockRange::ArabicExtendedA);
-                BlockRangesForLanguage.Add(EUnicodeBlockRange::ArabicMathematicalAlphabeticSymbols);
-                BlockRangesForLanguage.Add(EUnicodeBlockRange::ArabicPresentationFormsA);
-                BlockRangesForLanguage.Add(EUnicodeBlockRange::ArabicPresentationFormsB);
-                BlockRangesForLanguage.Add(EUnicodeBlockRange::ArabicSupplement);
-            }
-            //Japanese
-            else if (Language.Equals(TEXT("ja"), ESearchCase::IgnoreCase))
-            {
-                BlockRangesForLanguage.Add(EUnicodeBlockRange::CJKCompatibility);
-                BlockRangesForLanguage.Add(EUnicodeBlockRange::CJKCompatibilityForms);
-                BlockRangesForLanguage.Add(EUnicodeBlockRange::CJKCompatibilityIdeographs);
-                BlockRangesForLanguage.Add(EUnicodeBlockRange::CJKCompatibilityIdeographsSupplement);
-                BlockRangesForLanguage.Add(EUnicodeBlockRange::CJKRadicalsSupplement);
-                BlockRangesForLanguage.Add(EUnicodeBlockRange::CJKStrokes);
-                BlockRangesForLanguage.Add(EUnicodeBlockRange::CJKSymbolsAndPunctuation);
-                BlockRangesForLanguage.Add(EUnicodeBlockRange::CJKUnifiedIdeographs);
-                BlockRangesForLanguage.Add(EUnicodeBlockRange::CJKUnifiedIdeographsExtensionA);
-                BlockRangesForLanguage.Add(EUnicodeBlockRange::CJKUnifiedIdeographsExtensionB);
-                BlockRangesForLanguage.Add(EUnicodeBlockRange::CJKUnifiedIdeographsExtensionC);
-                BlockRangesForLanguage.Add(EUnicodeBlockRange::CJKUnifiedIdeographsExtensionD);
-                BlockRangesForLanguage.Add(EUnicodeBlockRange::CJKUnifiedIdeographsExtensionE);
-                BlockRangesForLanguage.Add(EUnicodeBlockRange::EnclosedCJKLettersAndMonths);
-                BlockRangesForLanguage.Add(EUnicodeBlockRange::Hiragana);
-                BlockRangesForLanguage.Add(EUnicodeBlockRange::Katakana);
-                BlockRangesForLanguage.Add(EUnicodeBlockRange::KatakanaPhoneticExtensions);
-                BlockRangesForLanguage.Add(EUnicodeBlockRange::Kanbun);
-                BlockRangesForLanguage.Add(EUnicodeBlockRange::HalfwidthAndFullwidthForms);
-            }
-            //Korean
-            else if (Language.Equals(TEXT("ko"), ESearchCase::IgnoreCase))
-            {
-                BlockRangesForLanguage.Add(EUnicodeBlockRange::HangulJamo);
-                BlockRangesForLanguage.Add(EUnicodeBlockRange::HangulJamoExtendedA);
-                BlockRangesForLanguage.Add(EUnicodeBlockRange::HangulJamoExtendedB);
-                BlockRangesForLanguage.Add(EUnicodeBlockRange::HangulCompatibilityJamo);
-                BlockRangesForLanguage.Add(EUnicodeBlockRange::HangulSyllables);
-            }
-            //Simplified Chinese
-            else if (Language.Equals(TEXT("zh-hans"), ESearchCase::IgnoreCase))
-            {
-                BlockRangesForLanguage.Add(EUnicodeBlockRange::CJKCompatibility);
-                BlockRangesForLanguage.Add(EUnicodeBlockRange::CJKCompatibilityForms);
-                BlockRangesForLanguage.Add(EUnicodeBlockRange::CJKCompatibilityIdeographs);
-                BlockRangesForLanguage.Add(EUnicodeBlockRange::CJKCompatibilityIdeographsSupplement);
-                BlockRangesForLanguage.Add(EUnicodeBlockRange::CJKRadicalsSupplement);
-                BlockRangesForLanguage.Add(EUnicodeBlockRange::CJKStrokes);
-                BlockRangesForLanguage.Add(EUnicodeBlockRange::CJKSymbolsAndPunctuation);
-                BlockRangesForLanguage.Add(EUnicodeBlockRange::CJKUnifiedIdeographs);
-                BlockRangesForLanguage.Add(EUnicodeBlockRange::CJKUnifiedIdeographsExtensionA);
-                BlockRangesForLanguage.Add(EUnicodeBlockRange::CJKUnifiedIdeographsExtensionB);
-                BlockRangesForLanguage.Add(EUnicodeBlockRange::CJKUnifiedIdeographsExtensionC);
-                BlockRangesForLanguage.Add(EUnicodeBlockRange::CJKUnifiedIdeographsExtensionD);
-                BlockRangesForLanguage.Add(EUnicodeBlockRange::CJKUnifiedIdeographsExtensionE);
-                BlockRangesForLanguage.Add(EUnicodeBlockRange::EnclosedCJKLettersAndMonths);
-            }
-            //Traditional Chinese
-            else if (Language.Equals(TEXT("zh-hant"), ESearchCase::IgnoreCase))
-            {
-                BlockRangesForLanguage.Add(EUnicodeBlockRange::CJKCompatibility);
-                BlockRangesForLanguage.Add(EUnicodeBlockRange::CJKCompatibilityForms);
-                BlockRangesForLanguage.Add(EUnicodeBlockRange::CJKCompatibilityIdeographs);
-                BlockRangesForLanguage.Add(EUnicodeBlockRange::CJKCompatibilityIdeographsSupplement);
-                BlockRangesForLanguage.Add(EUnicodeBlockRange::CJKRadicalsSupplement);
-                BlockRangesForLanguage.Add(EUnicodeBlockRange::CJKStrokes);
-                BlockRangesForLanguage.Add(EUnicodeBlockRange::CJKSymbolsAndPunctuation);
-                BlockRangesForLanguage.Add(EUnicodeBlockRange::CJKUnifiedIdeographs);
-                BlockRangesForLanguage.Add(EUnicodeBlockRange::CJKUnifiedIdeographsExtensionA);
-                BlockRangesForLanguage.Add(EUnicodeBlockRange::CJKUnifiedIdeographsExtensionB);
-                BlockRangesForLanguage.Add(EUnicodeBlockRange::CJKUnifiedIdeographsExtensionC);
-                BlockRangesForLanguage.Add(EUnicodeBlockRange::CJKUnifiedIdeographsExtensionD);
-                BlockRangesForLanguage.Add(EUnicodeBlockRange::CJKUnifiedIdeographsExtensionE);
-                BlockRangesForLanguage.Add(EUnicodeBlockRange::EnclosedCJKLettersAndMonths);
-            }
-
-            //Build out actual sub font ranges
-            FCompositeSubFont& SubFont = FontToBuild->SubTypefaces[FontToBuild->SubTypefaces.AddDefaulted()];
-            SubFont.Cultures.Append(Language);
-            for (EUnicodeBlockRange& BlockRange : BlockRangesForLanguage)
-            {
-                SubFont.CharacterRanges.Add(FUnicodeBlockRange::GetUnicodeBlockRange(BlockRange).Range);
-            }
-
-            //Finally append actual font
-            SubFont.Typeface.AppendFont(*FontIdentifier, FilePath, EFontHinting::Default, EFontLoadingPolicy::LazyLoad);
-        }
-    }
+void FPreLoadSettingsContainerBase::BuildSystemFont(const FString& FontIdentifer)
+{
+	if (!CachedSystemFont.IsValid())
+	{
+		const TArray<uint8> FontBytes = FPlatformMisc::GetSystemFontBytes();
+		if (FontBytes.Num() > 0)
+		{
+			const FString FontFilename = FPaths::EngineIntermediateDir() / TEXT("DefaultSystemFont.ttf");
+			if (FFileHelper::SaveArrayToFile(FontBytes, *FontFilename))
+			{
+				CachedSystemFont = MakeShared<FStandaloneCompositeFont>(NAME_None, FontFilename, EFontHinting::Default, EFontLoadingPolicy::LazyLoad);
+			}
+		}
+	}
+
+	if (ensureAlwaysMsgf(CachedSystemFont.IsValid(), TEXT("Failed to create system font for %s"), *FontIdentifer))
+	{
+		FontResources.Add(*FontIdentifer, CachedSystemFont);
+	}
 }
 
 bool FPreLoadSettingsContainerBase::IsValidLocalizedTextConfigString(TArray<FString>& SplitConfigEntry)
@@ -293,14 +472,14 @@ void FPreLoadSettingsContainerBase::CreateCustomSlateImageBrush(const FString& I
 {
 	if (bShouldLoadBrushes)
 	{
-		BrushResources.Add(*Identifier, new FSlateDynamicImageBrush(*TexturePath, ImageDimensions));
+		BrushResources.FindOrAdd(*Identifier, new FSlateDynamicImageBrush(*TexturePath, ImageDimensions));
 
 		//Make sure this dynamic image resource is registered with the SlateApplication
 		FSlateApplication::Get().GetRenderer()->GenerateDynamicImageResource(*TexturePath);
 	}
 	else
 	{
-		BrushResources.Add(*Identifier, new FSlateDynamicImageBrush(NAME_None, ImageDimensions));
+		BrushResources.FindOrAdd(*Identifier, new FSlateDynamicImageBrush(NAME_None, ImageDimensions));
 	}
 }
 
@@ -335,4 +514,14 @@ TSharedPtr<FCompositeFont> FPreLoadSettingsContainerBase::GetFont(const FString&
 FPreLoadSettingsContainerBase::FScreenGroupingBase* FPreLoadSettingsContainerBase::GetScreenGrouping(const FString& Identifier)
 {
     return ScreenGroupings.Find(*Identifier);
+}
+
+const FPreLoadSettingsContainerBase::FScreenGroupingBase* FPreLoadSettingsContainerBase::GetScreenAtIndex(int index) const
+{ 
+	return IsValidScreenIndex(index) ? ScreenGroupings.Find(ScreenOrderByLoadingGroups.Find(CurrentLoadGroup)->ScreenGroupings[index]) : nullptr;
+}
+
+bool FPreLoadSettingsContainerBase::IsValidScreenIndex(int index) const
+{
+	return (!CurrentLoadGroup.IsNone() && CurrentLoadGroup.IsValid() && (ScreenOrderByLoadingGroups.Find(CurrentLoadGroup) != nullptr)) ? ScreenOrderByLoadingGroups.Find(CurrentLoadGroup)->ScreenGroupings.IsValidIndex(index) : false;
 }
\ No newline at end of file
diff --git a/Engine/Source/Runtime/PreLoadScreen/Public/PreLoadSettingsContainer.h b/Engine/Source/Runtime/PreLoadScreen/Public/PreLoadSettingsContainer.h
--- a/Engine/Source/Runtime/PreLoadScreen/Public/PreLoadSettingsContainer.h
+++ b/Engine/Source/Runtime/PreLoadScreen/Public/PreLoadSettingsContainer.h
@@ -15,9 +15,9 @@ class PRELOADSCREEN_API FPreLoadSettingsContainerBase : public FDeferredCleanupI
 {
 public:
 
-    //Helper class to store groups of things we want to display together in the UI so that we can parse it easily in the .ini. 
+    //Helper struct to store groups of things we want to display together in the UI so that we can parse it easily in the .ini. 
     //IE: Show this background, with this text at this font size
-    class FScreenGroupingBase
+    struct FScreenGroupingBase
     {
     public:
         
@@ -32,6 +32,41 @@ public:
         {}
     };
 
+	//This is a listing of ScreenGroupings (stored by Identifier) that should be displayed in this order during a particular LoadingGroup.
+	//@TODO: TRoss, possible to move these loading groups into their own DeferreedCleanupInterface instead of the entire container being the DeferredCleanupInterface,
+	//if we want to support unloading them. For right now though, we mostly just care about loading them selectively and are ok with keeping them in memory until we clean everything up.
+	struct FScreenOrderByLoadingGroup
+	{
+	public:
+		TArray<FName> ScreenGroupings;
+
+		FScreenOrderByLoadingGroup()
+			: ScreenGroupings ()
+		{}
+	};
+
+	//Helper struct  to store information required to construct a CustomSlateImageBrush. Parsed from our .ini
+	struct FCustomBrushDefine
+	{
+	public:
+		FString BrushIdentifier;
+		FString FilePath;
+		FVector2D Size;
+
+		FCustomBrushDefine(const FString& BrushIdentifierIn, const FString& FilePathIn, FVector2D SizeIn)
+			: BrushIdentifier(BrushIdentifierIn)
+			, FilePath(FilePathIn)
+			, Size(SizeIn)
+		{}
+	};
+
+	//Helper struct to store all BrushDefines we need to load for a given BrushLoadingGroup
+	struct FCustomBrushLoadingGroup
+	{
+	public:
+		TArray<FCustomBrushDefine> CustomBrushDefinesToLoad;
+	};
+
 public:
 
     static FPreLoadSettingsContainerBase& Get()
@@ -56,6 +91,7 @@ public:
 public:
 
     FPreLoadSettingsContainerBase() 
+		: CurrentLoadGroup(NAME_None)
     {
 		bShouldLoadBrushes = true;
     }
@@ -71,8 +107,8 @@ public:
 
     int GetNumScreenGroupings() const { return ScreenGroupings.Num(); }
 
-    virtual const FScreenGroupingBase* GetScreenAtIndex(int index) { return IsValidScreenIndex(index) ? &ScreenGroupings[*ScreenDisplayOrder[index]] : nullptr; }
-    virtual bool IsValidScreenIndex(int index) { return ScreenDisplayOrder.IsValidIndex(index); }
+	virtual const FScreenGroupingBase* GetScreenAtIndex(int index) const;
+	virtual bool IsValidScreenIndex(int index) const;
 
     virtual void CreateCustomSlateImageBrush(const FString& Identifier, const FString& TexturePath, const FVector2D& ImageDimensions);
     virtual void AddLocalizedText(const FString& Identifier, FText LocalizedText);
@@ -81,13 +117,20 @@ public:
     //Maps the given font file to the given language and stores it under the FontIdentifier.
     //Identifier maps the entire CompositeFont, so if you want to add multiple fonts  for multiple languages, just store them all under the same identifer
     virtual void BuildCustomFont(const FString& FontIdentifier, const FString& Language, const FString& FilePath);
+	virtual void BuildSystemFont(const FString& FontIdentifer);
 
     //Helper functions that parse a .ini config entry and call the appropriate create function to 
     virtual void ParseBrushConfigEntry(const FString& BrushConfigEntry);
     virtual void ParseFontConfigEntry(const FString&  SplitConfigEntry);
     virtual void ParseLocalizedTextConfigString(const FString&  SplitConfigEntry);
     virtual void ParseScreenGroupingConfigString(const FString&  SplitConfigEntry);
-    
+
+	//Helper function to parse all .ini entries for LoadingGroups and ScreenOrder. Do these together so we can assert if
+	//we don't find a matching LoadingGroup identifier in the config. Should be run after we parse all screen groupings
+	virtual void ParseLoadingGroups(TArray<FString>& LoadingGroupIdentifiers);
+	virtual void ParseAllScreenOrderEntries(TArray<FString>& LoadingGroups, TArray<FString>& ScreenOrderEntries);
+	virtual void ParseScreenOrderConfigString(const FString& ScreenOrderEntry);
+
     //Sets the PluginContent dir so that when parsing config entries we can accept plugin relative file paths
     virtual void SetPluginContentDir(const FString& PluginContentDirIn) { PluginContentDir = PluginContentDirIn; }
 
@@ -96,8 +139,9 @@ public:
 
     float TimeToDisplayEachBackground;
     
-    //Screens are displayed in the order of this array
-    TArray<FString> ScreenDisplayOrder;
+	FName GetCurrentLoadGrouping() const { return CurrentLoadGroup; }
+	void LoadGrouping(FName Identifier);
+	void PerformInitialAssetLoad();
 
     //Helper function that takes in a file path and tries to reconsile it to be Plugin Specific if applicable.
     //Ensures if file is not found in either Plugin's content dir or the original path
@@ -112,19 +156,34 @@ protected:
     virtual bool IsValidScreenGrooupingConfigString(TArray<FString>& SplitConfigEntry);
 
 protected:
-    
-    /* Property Storage. Ties FName to a particular resource so we can get it by identifier. */
+	TArray<FString> ParsedLoadingGroupIdentifiers;
+
+	/* Property Storage. Ties FName to a particular resource so we can get it by identifier. */
     TMap<FName, const FSlateDynamicImageBrush*> BrushResources;
     TMap<FName, FText> LocalizedTextResources;
     TMap<FName, TSharedPtr<FCompositeFont>> FontResources;
+
+	TMap<FName, FScreenOrderByLoadingGroup> ScreenOrderByLoadingGroups;
     TMap<FName, FScreenGroupingBase> ScreenGroupings;
 
+	TMap<FName, FCustomBrushLoadingGroup> BrushLoadingGroups;
+
     //This string is used to make file paths relative to a particular Plugin's content directory when parsing file paths.
     FString PluginContentDir;
 
 	// Rather we should load image brushes
 	bool bShouldLoadBrushes;
 
+	FName CurrentLoadGroup;
+
+	TSharedPtr<FCompositeFont> CachedSystemFont;
+
+	//If our Font filepath is set to this, we use the system font instead of a custom font we load in
+	static FString UseSystemFontOverride;
+
+	//If we supply no loading groups, use this identifier by default
+	static FString DefaultInitialLoadingGroupIdentifier;
+
     // Singleton Instance -- This is only not a TSharedPtr as it needs to be cleaned up by a deferredcleanup call which directly
     // nukes the underlying object, causing a SharedPtr crash at shutdown.
     static FPreLoadSettingsContainerBase* Instance;
