PreLoadSettingsContainer changed to allow you to use system font per-language.
#test Android Client

diff --git a/Engine/Source/Runtime/PreLoadScreen/Private/PreLoadSettingsContainer.cpp b/Engine/Source/Runtime/PreLoadScreen/Private/PreLoadSettingsContainer.cpp
--- a/Engine/Source/Runtime/PreLoadScreen/Private/PreLoadSettingsContainer.cpp
+++ b/Engine/Source/Runtime/PreLoadScreen/Private/PreLoadSettingsContainer.cpp
@@ -231,144 +231,146 @@ void FPreLoadSettingsContainerBase::LoadGrouping(FName Identifier)
 
 void FPreLoadSettingsContainerBase::BuildCustomFont(const FString& FontIdentifier, const FString& Language, const FString& FilePath)
 {
-	if (FilePath == UseSystemFontOverride)
+	const bool bIsUsingSystemFont = (FilePath == UseSystemFontOverride);
+	
+	//Make sure we have created a System font file to load
+	if (bIsUsingSystemFont)
 	{
-		BuildSystemFont(FontIdentifier);
+		BuildSystemFontFile();
 	}
-	else
+
+	//Try and find existing font, if we can't, make a new one
+	TSharedPtr<FCompositeFont>& FontToBuild = FontResources.FindOrAdd(*FontIdentifier);
+	if (!FontToBuild.IsValid())
+	{
+		FontToBuild = MakeShared<FCompositeFont>();
+	}
+
+	if (ensureAlwaysMsgf(FontToBuild.IsValid(), TEXT("Error creating custom font!")))
 	{
-		//Try and find existing font, if we can't, make a new one
-		TSharedPtr<FCompositeFont>& FontToBuild = FontResources.FindOrAdd(*FontIdentifier);
-		if (!FontToBuild.IsValid())
+		//Overwrite FilePath if using the system font
+		FString FilePathToUse = bIsUsingSystemFont ? GetSystemFontFilePath() : FilePath;
+
+		//If En, then setup as default font
+		if (Language.Equals(TEXT("en")))
 		{
-			FontToBuild = MakeShared<FCompositeFont>();
+			FontToBuild->DefaultTypeface.AppendFont(*FontIdentifier, FilePathToUse, EFontHinting::Default, EFontLoadingPolicy::LazyLoad);
 		}
-
-		if (ensureAlwaysMsgf(FontToBuild.IsValid(), TEXT("Error creating custom font!")))
+		// if not en, we need to setup some block ranges and subfonts to handle special characters
+		else
 		{
-
-			//If En, then setup as default font
-			if (Language.Equals(TEXT("en")))
+			TArray<EUnicodeBlockRange> BlockRangesForLanguage;
+			//Arabic
+			if (Language.Equals(TEXT("ar"), ESearchCase::IgnoreCase))
 			{
-				FontToBuild->DefaultTypeface.AppendFont(*FontIdentifier, FilePath, EFontHinting::Default, EFontLoadingPolicy::LazyLoad);
+				BlockRangesForLanguage.Add(EUnicodeBlockRange::Arabic);
+				BlockRangesForLanguage.Add(EUnicodeBlockRange::ArabicExtendedA);
+				BlockRangesForLanguage.Add(EUnicodeBlockRange::ArabicMathematicalAlphabeticSymbols);
+				BlockRangesForLanguage.Add(EUnicodeBlockRange::ArabicPresentationFormsA);
+				BlockRangesForLanguage.Add(EUnicodeBlockRange::ArabicPresentationFormsB);
+				BlockRangesForLanguage.Add(EUnicodeBlockRange::ArabicSupplement);
 			}
-			// if not en, we need to setup some block ranges and subfonts to handle special characters
-			else
+			//Japanese
+			else if (Language.Equals(TEXT("ja"), ESearchCase::IgnoreCase))
 			{
-				TArray<EUnicodeBlockRange> BlockRangesForLanguage;
-				//Arabic
-				if (Language.Equals(TEXT("ar"), ESearchCase::IgnoreCase))
-				{
-					BlockRangesForLanguage.Add(EUnicodeBlockRange::Arabic);
-					BlockRangesForLanguage.Add(EUnicodeBlockRange::ArabicExtendedA);
-					BlockRangesForLanguage.Add(EUnicodeBlockRange::ArabicMathematicalAlphabeticSymbols);
-					BlockRangesForLanguage.Add(EUnicodeBlockRange::ArabicPresentationFormsA);
-					BlockRangesForLanguage.Add(EUnicodeBlockRange::ArabicPresentationFormsB);
-					BlockRangesForLanguage.Add(EUnicodeBlockRange::ArabicSupplement);
-				}
-				//Japanese
-				else if (Language.Equals(TEXT("ja"), ESearchCase::IgnoreCase))
-				{
-					BlockRangesForLanguage.Add(EUnicodeBlockRange::CJKCompatibility);
-					BlockRangesForLanguage.Add(EUnicodeBlockRange::CJKCompatibilityForms);
-					BlockRangesForLanguage.Add(EUnicodeBlockRange::CJKCompatibilityIdeographs);
-					BlockRangesForLanguage.Add(EUnicodeBlockRange::CJKCompatibilityIdeographsSupplement);
-					BlockRangesForLanguage.Add(EUnicodeBlockRange::CJKRadicalsSupplement);
-					BlockRangesForLanguage.Add(EUnicodeBlockRange::CJKStrokes);
-					BlockRangesForLanguage.Add(EUnicodeBlockRange::CJKSymbolsAndPunctuation);
-					BlockRangesForLanguage.Add(EUnicodeBlockRange::CJKUnifiedIdeographs);
-					BlockRangesForLanguage.Add(EUnicodeBlockRange::CJKUnifiedIdeographsExtensionA);
-					BlockRangesForLanguage.Add(EUnicodeBlockRange::CJKUnifiedIdeographsExtensionB);
-					BlockRangesForLanguage.Add(EUnicodeBlockRange::CJKUnifiedIdeographsExtensionC);
-					BlockRangesForLanguage.Add(EUnicodeBlockRange::CJKUnifiedIdeographsExtensionD);
-					BlockRangesForLanguage.Add(EUnicodeBlockRange::CJKUnifiedIdeographsExtensionE);
-					BlockRangesForLanguage.Add(EUnicodeBlockRange::EnclosedCJKLettersAndMonths);
-					BlockRangesForLanguage.Add(EUnicodeBlockRange::Hiragana);
-					BlockRangesForLanguage.Add(EUnicodeBlockRange::Katakana);
-					BlockRangesForLanguage.Add(EUnicodeBlockRange::KatakanaPhoneticExtensions);
-					BlockRangesForLanguage.Add(EUnicodeBlockRange::Kanbun);
-					BlockRangesForLanguage.Add(EUnicodeBlockRange::HalfwidthAndFullwidthForms);
-				}
-				//Korean
-				else if (Language.Equals(TEXT("ko"), ESearchCase::IgnoreCase))
-				{
-					BlockRangesForLanguage.Add(EUnicodeBlockRange::HangulJamo);
-					BlockRangesForLanguage.Add(EUnicodeBlockRange::HangulJamoExtendedA);
-					BlockRangesForLanguage.Add(EUnicodeBlockRange::HangulJamoExtendedB);
-					BlockRangesForLanguage.Add(EUnicodeBlockRange::HangulCompatibilityJamo);
-					BlockRangesForLanguage.Add(EUnicodeBlockRange::HangulSyllables);
-				}
-				//Simplified Chinese
-				else if (Language.Equals(TEXT("zh-hans"), ESearchCase::IgnoreCase))
-				{
-					BlockRangesForLanguage.Add(EUnicodeBlockRange::CJKCompatibility);
-					BlockRangesForLanguage.Add(EUnicodeBlockRange::CJKCompatibilityForms);
-					BlockRangesForLanguage.Add(EUnicodeBlockRange::CJKCompatibilityIdeographs);
-					BlockRangesForLanguage.Add(EUnicodeBlockRange::CJKCompatibilityIdeographsSupplement);
-					BlockRangesForLanguage.Add(EUnicodeBlockRange::CJKRadicalsSupplement);
-					BlockRangesForLanguage.Add(EUnicodeBlockRange::CJKStrokes);
-					BlockRangesForLanguage.Add(EUnicodeBlockRange::CJKSymbolsAndPunctuation);
-					BlockRangesForLanguage.Add(EUnicodeBlockRange::CJKUnifiedIdeographs);
-					BlockRangesForLanguage.Add(EUnicodeBlockRange::CJKUnifiedIdeographsExtensionA);
-					BlockRangesForLanguage.Add(EUnicodeBlockRange::CJKUnifiedIdeographsExtensionB);
-					BlockRangesForLanguage.Add(EUnicodeBlockRange::CJKUnifiedIdeographsExtensionC);
-					BlockRangesForLanguage.Add(EUnicodeBlockRange::CJKUnifiedIdeographsExtensionD);
-					BlockRangesForLanguage.Add(EUnicodeBlockRange::CJKUnifiedIdeographsExtensionE);
-					BlockRangesForLanguage.Add(EUnicodeBlockRange::EnclosedCJKLettersAndMonths);
-				}
-				//Traditional Chinese
-				else if (Language.Equals(TEXT("zh-hant"), ESearchCase::IgnoreCase))
-				{
-					BlockRangesForLanguage.Add(EUnicodeBlockRange::CJKCompatibility);
-					BlockRangesForLanguage.Add(EUnicodeBlockRange::CJKCompatibilityForms);
-					BlockRangesForLanguage.Add(EUnicodeBlockRange::CJKCompatibilityIdeographs);
-					BlockRangesForLanguage.Add(EUnicodeBlockRange::CJKCompatibilityIdeographsSupplement);
-					BlockRangesForLanguage.Add(EUnicodeBlockRange::CJKRadicalsSupplement);
-					BlockRangesForLanguage.Add(EUnicodeBlockRange::CJKStrokes);
-					BlockRangesForLanguage.Add(EUnicodeBlockRange::CJKSymbolsAndPunctuation);
-					BlockRangesForLanguage.Add(EUnicodeBlockRange::CJKUnifiedIdeographs);
-					BlockRangesForLanguage.Add(EUnicodeBlockRange::CJKUnifiedIdeographsExtensionA);
-					BlockRangesForLanguage.Add(EUnicodeBlockRange::CJKUnifiedIdeographsExtensionB);
-					BlockRangesForLanguage.Add(EUnicodeBlockRange::CJKUnifiedIdeographsExtensionC);
-					BlockRangesForLanguage.Add(EUnicodeBlockRange::CJKUnifiedIdeographsExtensionD);
-					BlockRangesForLanguage.Add(EUnicodeBlockRange::CJKUnifiedIdeographsExtensionE);
-					BlockRangesForLanguage.Add(EUnicodeBlockRange::EnclosedCJKLettersAndMonths);
-				}
-
-				//Build out actual sub font ranges
-				FCompositeSubFont& SubFont = FontToBuild->SubTypefaces[FontToBuild->SubTypefaces.AddDefaulted()];
-				SubFont.Cultures.Append(Language);
-				for (EUnicodeBlockRange& BlockRange : BlockRangesForLanguage)
-				{
-					SubFont.CharacterRanges.Add(FUnicodeBlockRange::GetUnicodeBlockRange(BlockRange).Range);
-				}
+				BlockRangesForLanguage.Add(EUnicodeBlockRange::CJKCompatibility);
+				BlockRangesForLanguage.Add(EUnicodeBlockRange::CJKCompatibilityForms);
+				BlockRangesForLanguage.Add(EUnicodeBlockRange::CJKCompatibilityIdeographs);
+				BlockRangesForLanguage.Add(EUnicodeBlockRange::CJKCompatibilityIdeographsSupplement);
+				BlockRangesForLanguage.Add(EUnicodeBlockRange::CJKRadicalsSupplement);
+				BlockRangesForLanguage.Add(EUnicodeBlockRange::CJKStrokes);
+				BlockRangesForLanguage.Add(EUnicodeBlockRange::CJKSymbolsAndPunctuation);
+				BlockRangesForLanguage.Add(EUnicodeBlockRange::CJKUnifiedIdeographs);
+				BlockRangesForLanguage.Add(EUnicodeBlockRange::CJKUnifiedIdeographsExtensionA);
+				BlockRangesForLanguage.Add(EUnicodeBlockRange::CJKUnifiedIdeographsExtensionB);
+				BlockRangesForLanguage.Add(EUnicodeBlockRange::CJKUnifiedIdeographsExtensionC);
+				BlockRangesForLanguage.Add(EUnicodeBlockRange::CJKUnifiedIdeographsExtensionD);
+				BlockRangesForLanguage.Add(EUnicodeBlockRange::CJKUnifiedIdeographsExtensionE);
+				BlockRangesForLanguage.Add(EUnicodeBlockRange::EnclosedCJKLettersAndMonths);
+				BlockRangesForLanguage.Add(EUnicodeBlockRange::Hiragana);
+				BlockRangesForLanguage.Add(EUnicodeBlockRange::Katakana);
+				BlockRangesForLanguage.Add(EUnicodeBlockRange::KatakanaPhoneticExtensions);
+				BlockRangesForLanguage.Add(EUnicodeBlockRange::Kanbun);
+				BlockRangesForLanguage.Add(EUnicodeBlockRange::HalfwidthAndFullwidthForms);
+			}
+			//Korean
+			else if (Language.Equals(TEXT("ko"), ESearchCase::IgnoreCase))
+			{
+				BlockRangesForLanguage.Add(EUnicodeBlockRange::HangulJamo);
+				BlockRangesForLanguage.Add(EUnicodeBlockRange::HangulJamoExtendedA);
+				BlockRangesForLanguage.Add(EUnicodeBlockRange::HangulJamoExtendedB);
+				BlockRangesForLanguage.Add(EUnicodeBlockRange::HangulCompatibilityJamo);
+				BlockRangesForLanguage.Add(EUnicodeBlockRange::HangulSyllables);
+			}
+			//Simplified Chinese
+			else if (Language.Equals(TEXT("zh-hans"), ESearchCase::IgnoreCase))
+			{
+				BlockRangesForLanguage.Add(EUnicodeBlockRange::CJKCompatibility);
+				BlockRangesForLanguage.Add(EUnicodeBlockRange::CJKCompatibilityForms);
+				BlockRangesForLanguage.Add(EUnicodeBlockRange::CJKCompatibilityIdeographs);
+				BlockRangesForLanguage.Add(EUnicodeBlockRange::CJKCompatibilityIdeographsSupplement);
+				BlockRangesForLanguage.Add(EUnicodeBlockRange::CJKRadicalsSupplement);
+				BlockRangesForLanguage.Add(EUnicodeBlockRange::CJKStrokes);
+				BlockRangesForLanguage.Add(EUnicodeBlockRange::CJKSymbolsAndPunctuation);
+				BlockRangesForLanguage.Add(EUnicodeBlockRange::CJKUnifiedIdeographs);
+				BlockRangesForLanguage.Add(EUnicodeBlockRange::CJKUnifiedIdeographsExtensionA);
+				BlockRangesForLanguage.Add(EUnicodeBlockRange::CJKUnifiedIdeographsExtensionB);
+				BlockRangesForLanguage.Add(EUnicodeBlockRange::CJKUnifiedIdeographsExtensionC);
+				BlockRangesForLanguage.Add(EUnicodeBlockRange::CJKUnifiedIdeographsExtensionD);
+				BlockRangesForLanguage.Add(EUnicodeBlockRange::CJKUnifiedIdeographsExtensionE);
+				BlockRangesForLanguage.Add(EUnicodeBlockRange::EnclosedCJKLettersAndMonths);
+			}
+			//Traditional Chinese
+			else if (Language.Equals(TEXT("zh-hant"), ESearchCase::IgnoreCase))
+			{
+				BlockRangesForLanguage.Add(EUnicodeBlockRange::CJKCompatibility);
+				BlockRangesForLanguage.Add(EUnicodeBlockRange::CJKCompatibilityForms);
+				BlockRangesForLanguage.Add(EUnicodeBlockRange::CJKCompatibilityIdeographs);
+				BlockRangesForLanguage.Add(EUnicodeBlockRange::CJKCompatibilityIdeographsSupplement);
+				BlockRangesForLanguage.Add(EUnicodeBlockRange::CJKRadicalsSupplement);
+				BlockRangesForLanguage.Add(EUnicodeBlockRange::CJKStrokes);
+				BlockRangesForLanguage.Add(EUnicodeBlockRange::CJKSymbolsAndPunctuation);
+				BlockRangesForLanguage.Add(EUnicodeBlockRange::CJKUnifiedIdeographs);
+				BlockRangesForLanguage.Add(EUnicodeBlockRange::CJKUnifiedIdeographsExtensionA);
+				BlockRangesForLanguage.Add(EUnicodeBlockRange::CJKUnifiedIdeographsExtensionB);
+				BlockRangesForLanguage.Add(EUnicodeBlockRange::CJKUnifiedIdeographsExtensionC);
+				BlockRangesForLanguage.Add(EUnicodeBlockRange::CJKUnifiedIdeographsExtensionD);
+				BlockRangesForLanguage.Add(EUnicodeBlockRange::CJKUnifiedIdeographsExtensionE);
+				BlockRangesForLanguage.Add(EUnicodeBlockRange::EnclosedCJKLettersAndMonths);
+			}
 
-				//Finally append actual font
-				SubFont.Typeface.AppendFont(*FontIdentifier, FilePath, EFontHinting::Default, EFontLoadingPolicy::LazyLoad);
+			//Build out actual sub font ranges
+			FCompositeSubFont& SubFont = FontToBuild->SubTypefaces[FontToBuild->SubTypefaces.AddDefaulted()];
+			SubFont.Cultures.Append(Language);
+			for (EUnicodeBlockRange& BlockRange : BlockRangesForLanguage)
+			{
+				SubFont.CharacterRanges.Add(FUnicodeBlockRange::GetUnicodeBlockRange(BlockRange).Range);
 			}
+
+			//Finally append actual font
+			SubFont.Typeface.AppendFont(*FontIdentifier, FilePathToUse, EFontHinting::Default, EFontLoadingPolicy::LazyLoad);
 		}
 	}
 }
 
-void FPreLoadSettingsContainerBase::BuildSystemFont(const FString& FontIdentifer)
+const FString FPreLoadSettingsContainerBase::GetSystemFontFilePath() const
 {
-	if (!CachedSystemFont.IsValid())
+	return FPaths::EngineIntermediateDir() / TEXT("DefaultSystemFont.ttf");
+}
+
+bool FPreLoadSettingsContainerBase::BuildSystemFontFile()
+{
+	if (!HasCreatedSystemFontFile)
 	{
 		const TArray<uint8> FontBytes = FPlatformMisc::GetSystemFontBytes();
 		if (FontBytes.Num() > 0)
-		{
-			const FString FontFilename = FPaths::EngineIntermediateDir() / TEXT("DefaultSystemFont.ttf");
-			if (FFileHelper::SaveArrayToFile(FontBytes, *FontFilename))
-			{
-				CachedSystemFont = MakeShared<FStandaloneCompositeFont>(NAME_None, FontFilename, EFontHinting::Default, EFontLoadingPolicy::LazyLoad);
-			}
+		{			
+			HasCreatedSystemFontFile = FFileHelper::SaveArrayToFile(FontBytes, *GetSystemFontFilePath());
 		}
 	}
 
-	if (ensureAlwaysMsgf(CachedSystemFont.IsValid(), TEXT("Failed to create system font for %s"), *FontIdentifer))
-	{
-		FontResources.Add(*FontIdentifer, CachedSystemFont);
-	}
+	ensureAlwaysMsgf(HasCreatedSystemFontFile, TEXT("Failed to create system font!"));
+	return HasCreatedSystemFontFile;
 }
 
 bool FPreLoadSettingsContainerBase::IsValidLocalizedTextConfigString(TArray<FString>& SplitConfigEntry)
diff --git a/Engine/Source/Runtime/PreLoadScreen/Public/PreLoadSettingsContainer.h b/Engine/Source/Runtime/PreLoadScreen/Public/PreLoadSettingsContainer.h
--- a/Engine/Source/Runtime/PreLoadScreen/Public/PreLoadSettingsContainer.h
+++ b/Engine/Source/Runtime/PreLoadScreen/Public/PreLoadSettingsContainer.h
@@ -94,6 +94,7 @@ public:
 		: CurrentLoadGroup(NAME_None)
     {
 		bShouldLoadBrushes = true;
+		HasCreatedSystemFontFile = false;
     }
 
     virtual ~FPreLoadSettingsContainerBase();
@@ -117,7 +118,8 @@ public:
     //Maps the given font file to the given language and stores it under the FontIdentifier.
     //Identifier maps the entire CompositeFont, so if you want to add multiple fonts  for multiple languages, just store them all under the same identifer
     virtual void BuildCustomFont(const FString& FontIdentifier, const FString& Language, const FString& FilePath);
-	virtual void BuildSystemFont(const FString& FontIdentifer);
+	virtual bool BuildSystemFontFile();
+	virtual const FString GetSystemFontFilePath() const;
 
     //Helper functions that parse a .ini config entry and call the appropriate create function to 
     virtual void ParseBrushConfigEntry(const FString& BrushConfigEntry);
@@ -176,7 +178,7 @@ protected:
 
 	FName CurrentLoadGroup;
 
-	TSharedPtr<FCompositeFont> CachedSystemFont;
+	bool HasCreatedSystemFontFile;
 
 	//If our Font filepath is set to this, we use the system font instead of a custom font we load in
 	static FString UseSystemFontOverride;
