Adding ability to cancel install in BPS PreLoadManager

diff --git a/Engine/Source/Runtime/PreLoadScreen/Private/BuildPatchServicesPreLoadManager.cpp b/Engine/Source/Runtime/PreLoadScreen/Private/BuildPatchServicesPreLoadManager.cpp
--- a/Engine/Source/Runtime/PreLoadScreen/Private/BuildPatchServicesPreLoadManager.cpp
+++ b/Engine/Source/Runtime/PreLoadScreen/Private/BuildPatchServicesPreLoadManager.cpp
@@ -64,6 +64,14 @@ void FBuildPatchServicesPreLoadManagerBase::ResumeBuildPatchInstall()
     }
 }
 
+void FBuildPatchServicesPreLoadManagerBase::CancelBuildPatchInstall()
+{
+	if (ContentBuildInstaller.IsValid())
+	{
+		ContentBuildInstaller->CancelInstall();
+	}
+}
+
 const FText& FBuildPatchServicesPreLoadManagerBase::GetStatusText() const
 {
 	// Static const fixed FText values so that they are not constantly constructed
diff --git a/Engine/Source/Runtime/PreLoadScreen/Public/BuildPatchServicesPreLoadManager.h b/Engine/Source/Runtime/PreLoadScreen/Public/BuildPatchServicesPreLoadManager.h
--- a/Engine/Source/Runtime/PreLoadScreen/Public/BuildPatchServicesPreLoadManager.h
+++ b/Engine/Source/Runtime/PreLoadScreen/Public/BuildPatchServicesPreLoadManager.h
@@ -31,6 +31,7 @@ public:
 
     virtual void PauseBuildPatchInstall();
     virtual void ResumeBuildPatchInstall();
+	virtual void CancelBuildPatchInstall();
 
     float GetProgressPercent() const { return ContentBuildInstaller.IsValid() ? ContentBuildInstaller->GetUpdateProgress() : 0.0f; }
     EBuildPatchDownloadHealth GetDownloadHealth() const { return ContentBuildInstaller.IsValid() ? ContentBuildInstaller->GetDownloadHealth() : EBuildPatchDownloadHealth::NUM_Values; }
