Merge 4822556
Fix for stats crash.

diff --git a/Engine/Source/Runtime/PreLoadScreen/Private/PreLoadScreenManager.cpp b/Engine/Source/Runtime/PreLoadScreen/Private/PreLoadScreenManager.cpp
--- a/Engine/Source/Runtime/PreLoadScreen/Private/PreLoadScreenManager.cpp
+++ b/Engine/Source/Runtime/PreLoadScreen/Private/PreLoadScreenManager.cpp
@@ -11,6 +11,7 @@
 #include "PreLoadSettingsContainer.h"
 
 #include "HAL/ThreadManager.h"
+#include "Modules/ModuleManager.h"
 
 #if PLATFORM_ANDROID
 #if USE_ANDROID_EVENTS
@@ -18,6 +19,8 @@
 #endif
 #endif
 
+IMPLEMENT_MODULE(FDefaultModuleImpl, PreLoadScreen);
+
 DEFINE_LOG_CATEGORY_STATIC(LogPreLoadScreenManager, Log, All);
 
 TSharedPtr<FPreLoadScreenManager> FPreLoadScreenManager::Instance;
