-Allow preload screens to optionally render
-Add FileExistsInPlatformPackage() to platform misc
-Add HasBuildMetaData() to bundle manager interface

diff --git a/Engine/Plugins/Runtime/PreLoadScreenMoviePlayer/Source/PreLoadScreenMoviePlayer/Private/PreLoadMoviePlayerScreenBase.cpp b/Engine/Plugins/Runtime/PreLoadScreenMoviePlayer/Source/PreLoadScreenMoviePlayer/Private/PreLoadMoviePlayerScreenBase.cpp
--- a/Engine/Plugins/Runtime/PreLoadScreenMoviePlayer/Source/PreLoadScreenMoviePlayer/Private/PreLoadMoviePlayerScreenBase.cpp
+++ b/Engine/Plugins/Runtime/PreLoadScreenMoviePlayer/Source/PreLoadScreenMoviePlayer/Private/PreLoadMoviePlayerScreenBase.cpp
@@ -80,7 +80,7 @@ protected:
     FOnKeyDown OnKeyDownHandler;
 };
 
-void FPreLoadMoviePlayerScreenBase::Init()
+bool FPreLoadMoviePlayerScreenBase::Init()
 {
     if (!bInitialized)
     {
@@ -119,6 +119,8 @@ void FPreLoadMoviePlayerScreenBase::Init()
             FSlateApplication::Get().RegisterGameViewport(MovieViewport.ToSharedRef());
         }
     }
+
+	return true;
 }
 
 void FPreLoadMoviePlayerScreenBase::OnPlay(TWeakPtr<SWindow> TargetWindow)
diff --git a/Engine/Plugins/Runtime/PreLoadScreenMoviePlayer/Source/PreLoadScreenMoviePlayer/Public/PreLoadMoviePlayerScreenBase.h b/Engine/Plugins/Runtime/PreLoadScreenMoviePlayer/Source/PreLoadScreenMoviePlayer/Public/PreLoadMoviePlayerScreenBase.h
--- a/Engine/Plugins/Runtime/PreLoadScreenMoviePlayer/Source/PreLoadScreenMoviePlayer/Public/PreLoadMoviePlayerScreenBase.h
+++ b/Engine/Plugins/Runtime/PreLoadScreenMoviePlayer/Source/PreLoadScreenMoviePlayer/Public/PreLoadMoviePlayerScreenBase.h
@@ -15,21 +15,21 @@ class FPreLoadMoviePlayerScreenBase : public FPreLoadScreenBase, public TSharedF
 public:
 
     virtual void OnPlay(TWeakPtr<SWindow> TargetWindow) override;
-    virtual void OnStop();
-    virtual void Tick(float DeltaTime);
-    virtual void RenderTick(float DeltaTime);
+    virtual void OnStop() override;
+    virtual void Tick(float DeltaTime) override;
+    virtual void RenderTick(float DeltaTime) override;
 
-    virtual void Init();
+    virtual bool Init() override;
 
     virtual TSharedPtr<SWidget> GetWidget() override { return MoviePlayerContents; }
     virtual const TSharedPtr<const SWidget> GetWidget() const override { return MoviePlayerContents; }
 
-    virtual void CleanUp();
+    virtual void CleanUp() override;
 
     virtual void SetMovieAttributes(const FPreLoadMovieAttributes& MovieAttributesIn) { MovieAttributes = MovieAttributesIn; }
 
     //Default behavior is just to see if we have an active widget. Should really overload with our own behavior to see if we are done displaying
-    virtual bool IsDone() const { return IsMovieStreamingFinished(); }
+    virtual bool IsDone() const override { return IsMovieStreamingFinished(); }
 
     virtual bool MovieStreamingIsPrepared() const;
     virtual FReply OnLoadingScreenMouseButtonDown(const FGeometry& Geometry, const FPointerEvent& PointerEvent);
@@ -45,7 +45,7 @@ public:
     virtual bool IsMovieStreamingFinished() const;
 
     FPreLoadMoviePlayerScreenBase() {};
-    virtual ~FPreLoadMoviePlayerScreenBase() {};
+    virtual ~FPreLoadMoviePlayerScreenBase() override {};
     
     virtual void InitSettingsFromConfig(const FString& ConfigFileName) override;
 
diff --git a/Engine/Source/Runtime/Core/Private/Android/AndroidPlatformMisc.cpp b/Engine/Source/Runtime/Core/Private/Android/AndroidPlatformMisc.cpp
--- a/Engine/Source/Runtime/Core/Private/Android/AndroidPlatformMisc.cpp
+++ b/Engine/Source/Runtime/Core/Private/Android/AndroidPlatformMisc.cpp
@@ -1135,6 +1135,20 @@ FString FAndroidMisc::LoadTextFileFromPlatformPackage(const FString& RelativePat
 	return FString();
 }
 
+bool FAndroidMisc::FileExitsInPlatformPackage(const FString& RelativePath)
+{
+#if USE_ANDROID_JNI
+	AAssetManager* AssetMgr = AndroidThunkCpp_GetAssetManager();
+	AAsset* asset = AAssetManager_open(AssetMgr, TCHAR_TO_UTF8(*RelativePath), AASSET_MODE_UNKNOWN);
+	if (asset)
+	{
+		AAsset_close(asset);
+		return true;
+	}
+#endif
+	return false;
+}
+
 void FAndroidMisc::SetVersionInfo( FString InAndroidVersion, FString InDeviceMake, FString InDeviceModel, FString InDeviceBuildNumber, FString InOSLanguage )
 {
 	AndroidVersion = InAndroidVersion;
diff --git a/Engine/Source/Runtime/Core/Private/GenericPlatform/GenericPlatformMisc.cpp b/Engine/Source/Runtime/Core/Private/GenericPlatform/GenericPlatformMisc.cpp
--- a/Engine/Source/Runtime/Core/Private/GenericPlatform/GenericPlatformMisc.cpp
+++ b/Engine/Source/Runtime/Core/Private/GenericPlatform/GenericPlatformMisc.cpp
@@ -1350,6 +1350,12 @@ FString FGenericPlatformMisc::LoadTextFileFromPlatformPackage(const FString& Rel
 	return Result;
 }
 
+bool FGenericPlatformMisc::FileExitsInPlatformPackage(const FString& RelativePath)
+{
+	FString Path = RootDir() / RelativePath;
+	return IFileManager::Get().FileExists(*Path);
+}
+
 void FGenericPlatformMisc::ParseChunkIdPakchunkIndexMapping(TArray<FString> ChunkIndexMappingData, TMap<int32, int32>& OutMapping)
 {
 	OutMapping.Empty();
diff --git a/Engine/Source/Runtime/Core/Private/IOS/IOSPlatformMisc.cpp b/Engine/Source/Runtime/Core/Private/IOS/IOSPlatformMisc.cpp
--- a/Engine/Source/Runtime/Core/Private/IOS/IOSPlatformMisc.cpp
+++ b/Engine/Source/Runtime/Core/Private/IOS/IOSPlatformMisc.cpp
@@ -1090,6 +1090,13 @@ FString FIOSPlatformMisc::LoadTextFileFromPlatformPackage(const FString& Relativ
 	return FString(UTF8_TO_TCHAR(FileContents.GetData()));
 }
 
+bool FIOSPlatformMisc::FileExitsInPlatformPackage(const FString& RelativePath)
+{
+	FString FilePath = FString([[NSBundle mainBundle] bundlePath]) / RelativePath;
+
+	return 0 == access(TCHAR_TO_UTF8(*FilePath), F_OK);
+}
+
 void FIOSPlatformMisc::EnableVoiceChat(bool bEnable)
 {
 	return [[IOSAppDelegate GetDelegate] EnableVoiceChat:bEnable];
diff --git a/Engine/Source/Runtime/Core/Public/Android/AndroidPlatformMisc.h b/Engine/Source/Runtime/Core/Public/Android/AndroidPlatformMisc.h
--- a/Engine/Source/Runtime/Core/Public/Android/AndroidPlatformMisc.h
+++ b/Engine/Source/Runtime/Core/Public/Android/AndroidPlatformMisc.h
@@ -160,6 +160,7 @@ public:
 	static void ShareURL(const FString& URL, const FText& Description, int32 LocationHintX, int32 LocationHintY);
 
 	static FString LoadTextFileFromPlatformPackage(const FString& RelativePath);
+	static bool FileExitsInPlatformPackage(const FString& RelativePath);
 
 	// ANDROID ONLY:
 	static void SetVersionInfo(FString AndroidVersion, FString DeviceMake, FString DeviceModel, FString DeviceBuildNumber, FString OSLanguage);
diff --git a/Engine/Source/Runtime/Core/Public/GenericPlatform/GenericPlatformInstallBundleManager.h b/Engine/Source/Runtime/Core/Public/GenericPlatform/GenericPlatformInstallBundleManager.h
--- a/Engine/Source/Runtime/Core/Public/GenericPlatform/GenericPlatformInstallBundleManager.h
+++ b/Engine/Source/Runtime/Core/Public/GenericPlatform/GenericPlatformInstallBundleManager.h
@@ -286,6 +286,8 @@ public:
 
 	virtual ~IPlatformInstallBundleManager() {}
 
+	virtual bool HasBuildMetaData() const = 0;
+
 	virtual void PushInitErrorCallback(FInstallBundleManagerInitErrorHandler Callback) = 0;
 	virtual void PopInitErrorCallback() = 0;
 
diff --git a/Engine/Source/Runtime/Core/Public/GenericPlatform/GenericPlatformMisc.h b/Engine/Source/Runtime/Core/Public/GenericPlatform/GenericPlatformMisc.h
--- a/Engine/Source/Runtime/Core/Public/GenericPlatform/GenericPlatformMisc.h
+++ b/Engine/Source/Runtime/Core/Public/GenericPlatform/GenericPlatformMisc.h
@@ -1288,6 +1288,8 @@ public:
 	*/
 	static FString LoadTextFileFromPlatformPackage(const FString& RelativePath);
 
+	static bool FileExitsInPlatformPackage(const FString& RelativePath);
+
 	static void ParseChunkIdPakchunkIndexMapping(TArray<FString> ChunkIndexRedirects, TMap<int32, int32>& OutMapping);
 
 	static void PumpMessagesOutsideMainLoop()
diff --git a/Engine/Source/Runtime/Core/Public/IOS/IOSPlatformMisc.h b/Engine/Source/Runtime/Core/Public/IOS/IOSPlatformMisc.h
--- a/Engine/Source/Runtime/Core/Public/IOS/IOSPlatformMisc.h
+++ b/Engine/Source/Runtime/Core/Public/IOS/IOSPlatformMisc.h
@@ -78,6 +78,7 @@ struct CORE_API FIOSPlatformMisc : public FApplePlatformMisc
 	static void ShareURL(const FString& URL, const FText& Description, int32 LocationHintX, int32 LocationHintY);
 
 	static FString LoadTextFileFromPlatformPackage(const FString& RelativePath);
+	static bool FileExitsInPlatformPackage(const FString& RelativePath);
 
 	static void EnableVoiceChat(bool bEnable);
 	static bool IsVoiceChatEnabled();
diff --git a/Engine/Source/Runtime/Launch/Private/LaunchEngineLoop.cpp b/Engine/Source/Runtime/Launch/Private/LaunchEngineLoop.cpp
--- a/Engine/Source/Runtime/Launch/Private/LaunchEngineLoop.cpp
+++ b/Engine/Source/Runtime/Launch/Private/LaunchEngineLoop.cpp
@@ -2474,7 +2474,12 @@ int32 FEngineLoop::PreInit(const TCHAR* CmdLine)
 					if (FPreLoadScreenManager::Get()->HasRegisteredPreLoadScreenType(EPreLoadScreenTypes::EarlyStartupScreen))
 					{
 						// disable the splash before playing the early startup screen
-						FPlatformMisc::PlatformHandleSplashScreen(false);
+						FPreLoadScreenManager::IsResponsibleForRenderingDelegate.AddLambda(
+							[](bool bIsPreloadScreenManResponsibleForRendering) 
+							{
+								FPlatformMisc::PlatformHandleSplashScreen(!bIsPreloadScreenManResponsibleForRendering);
+							}
+						);
 	                    FPreLoadScreenManager::Get()->PlayFirstPreLoadScreen(EPreLoadScreenTypes::EarlyStartupScreen);
 	                }
 					else
diff --git a/Engine/Source/Runtime/NullInstallBundleManager/Source/NullInstallBundleManager.cpp b/Engine/Source/Runtime/NullInstallBundleManager/Source/NullInstallBundleManager.cpp
--- a/Engine/Source/Runtime/NullInstallBundleManager/Source/NullInstallBundleManager.cpp
+++ b/Engine/Source/Runtime/NullInstallBundleManager/Source/NullInstallBundleManager.cpp
@@ -5,6 +5,11 @@
 
 class FNullInstallBundleManager : public IPlatformInstallBundleManager
 {
+	virtual bool HasBuildMetaData() const override
+	{
+		return false;
+	}
+
 	virtual void PushInitErrorCallback(FInstallBundleManagerInitErrorHandler Callback) override
 	{
 	}
diff --git a/Engine/Source/Runtime/PreLoadScreen/Private/PreLoadScreenBase.cpp b/Engine/Source/Runtime/PreLoadScreen/Private/PreLoadScreenBase.cpp
--- a/Engine/Source/Runtime/PreLoadScreen/Private/PreLoadScreenBase.cpp
+++ b/Engine/Source/Runtime/PreLoadScreen/Private/PreLoadScreenBase.cpp
@@ -7,10 +7,6 @@
 
 #include "Misc/ConfigCacheIni.h"
 
-void FPreLoadScreenBase::Init()
-{
-}
-
 void FPreLoadScreenBase::InitSettingsFromConfig(const FString& ConfigFileName)
 {	
 	SCOPED_BOOT_TIMING("FPreLoadScreenBase::InitSettingsFromConfig");
diff --git a/Engine/Source/Runtime/PreLoadScreen/Private/PreLoadScreenManager.cpp b/Engine/Source/Runtime/PreLoadScreen/Private/PreLoadScreenManager.cpp
--- a/Engine/Source/Runtime/PreLoadScreen/Private/PreLoadScreenManager.cpp
+++ b/Engine/Source/Runtime/PreLoadScreen/Private/PreLoadScreenManager.cpp
@@ -32,6 +32,8 @@ TSharedPtr<FPreLoadScreenManager> FPreLoadScreenManager::Instance;
 FCriticalSection FPreLoadScreenManager::RenderingEnabledCriticalSection;
 bool FPreLoadScreenManager::bRenderingEnabled = true;
 
+FIsPreloadScreenResponsibleForRenderingMultiDelegate FPreLoadScreenManager::IsResponsibleForRenderingDelegate;
+
 void FPreLoadScreenManager::Initialize(FSlateRenderer& InSlateRenderer)
 {
     if (bInitialized || !ArePreLoadScreensEnabled())       
@@ -343,48 +345,66 @@ void FPreLoadScreenManager::EnableRendering(bool bEnabled)
 
 void FPreLoadScreenManager::EarlyPlayRenderFrameTick()
 {
-    if (!ShouldRender())
-    {
-        return;
-    }
+	bool bIsResponsibleForRendering_Local = true;
 
-    IPreLoadScreen* ActivePreLoadScreen = PreLoadScreens[ActivePreLoadScreenIndex].Get();
-    if (ensureAlwaysMsgf(ActivePreLoadScreen, TEXT("Invalid Active PreLoadScreen during EarlyPlayRenderFrameTick!")))
-    {
-        FSlateApplication& SlateApp = FSlateApplication::Get();
-        float SlateDeltaTime = SlateApp.GetDeltaTime();
+	if (!ShouldRender())
+	{
+		// bIsResponsibleForRendering = true;
+		// In this case FPreLoadScreenManager is responsible for rendering but choosing not to, probably because the
+		// app is not in the fourground.
+		return;
+	}
 
-        //Setup Slate Render Command
-        ENQUEUE_RENDER_COMMAND(BeginPreLoadScreenFrame)(
-			[ActivePreLoadScreen, SlateDeltaTime](FRHICommandListImmediate& RHICmdList)
-            {
-                if (FPreLoadScreenManager::ShouldRender())
-                {
-                    GFrameNumberRenderThread++;
-                    GRHICommandList.GetImmediateCommandList().BeginFrame();
+	IPreLoadScreen* ActivePreLoadScreen = PreLoadScreens[ActivePreLoadScreenIndex].Get();
+	if (ensureAlwaysMsgf(ActivePreLoadScreen, TEXT("Invalid Active PreLoadScreen during EarlyPlayRenderFrameTick!")))
+	{
+		if (!ActivePreLoadScreen->ShouldRender())
+		{
+			bIsResponsibleForRendering_Local = false;
+		}
 
-                    ActivePreLoadScreen->RenderTick(SlateDeltaTime);
-                }
-            }
-        );
+		if (bIsResponsibleForRendering_Local != bIsResponsibleForRendering)
+		{
+			bIsResponsibleForRendering = bIsResponsibleForRendering_Local;
+			IsResponsibleForRenderingDelegate.Broadcast(bIsResponsibleForRendering);
+		}
 
-        SlateApp.Tick();
+		if (bIsResponsibleForRendering_Local)
+		{
+			FSlateApplication& SlateApp = FSlateApplication::Get();
+			float SlateDeltaTime = SlateApp.GetDeltaTime();
 
-        // Synchronize the game thread and the render thread so that the render thread doesn't get too far behind.
-        SlateApp.GetRenderer()->Sync();
+			//Setup Slate Render Command
+			ENQUEUE_RENDER_COMMAND(BeginPreLoadScreenFrame)(
+				[ActivePreLoadScreen, SlateDeltaTime](FRHICommandListImmediate& RHICmdList)
+				{
+					if (FPreLoadScreenManager::ShouldRender())
+					{
+						GFrameNumberRenderThread++;
+						GRHICommandList.GetImmediateCommandList().BeginFrame();
 
-        ENQUEUE_RENDER_COMMAND(FinishPreLoadScreenFrame)(
-            [](FRHICommandListImmediate& RHICmdList)
-        {
-            if (FPreLoadScreenManager::ShouldRender())
-            {
-                GRHICommandList.GetImmediateCommandList().EndFrame();
-                GRHICommandList.GetImmediateCommandList().ImmediateFlush(EImmediateFlushType::FlushRHIThreadFlushResources);
-            }
-        });
+						ActivePreLoadScreen->RenderTick(SlateDeltaTime);
+					}
+				});
 
-        FlushRenderingCommands();
-    }
+			SlateApp.Tick();
+
+			// Synchronize the game thread and the render thread so that the render thread doesn't get too far behind.
+			SlateApp.GetRenderer()->Sync();
+
+			ENQUEUE_RENDER_COMMAND(FinishPreLoadScreenFrame)(
+				[](FRHICommandListImmediate& RHICmdList)
+				{
+					if (FPreLoadScreenManager::ShouldRender())
+					{
+						GRHICommandList.GetImmediateCommandList().EndFrame();
+						GRHICommandList.GetImmediateCommandList().ImmediateFlush(EImmediateFlushType::FlushRHIThreadFlushResources);
+					}
+				});
+
+			//FlushRenderingCommands();
+		}
+	}
 }
 
 void FPreLoadScreenManager::StopPreLoadScreen()
diff --git a/Engine/Source/Runtime/PreLoadScreen/Public/PreLoadScreen.h b/Engine/Source/Runtime/PreLoadScreen/Public/PreLoadScreen.h
--- a/Engine/Source/Runtime/PreLoadScreen/Public/PreLoadScreen.h
+++ b/Engine/Source/Runtime/PreLoadScreen/Public/PreLoadScreen.h
@@ -18,7 +18,7 @@ class IPreLoadScreen //: public TSharedFromThis<IPreLoadScreen, ESPMode::ThreadS
 public:
     virtual ~IPreLoadScreen() {};
 
-    virtual void Init() = 0;
+    virtual bool Init() = 0;
     
     //Standard tick that happens every frame
     virtual void Tick(float DeltaTime) = 0;
@@ -27,6 +27,9 @@ public:
     //keeps us from spinning super fast when we aren't doing much beyond loading data / etc on other threads.
     virtual float GetAddedTickDelay() = 0;
 
+	//Whether an EarlyStartupLoadScreen should render
+	virtual bool ShouldRender() const = 0;
+
     //This tick happens as part of the slate render tick during an EarlyStartupLoadScreen
     virtual void RenderTick(float DeltaTime) = 0;
     
diff --git a/Engine/Source/Runtime/PreLoadScreen/Public/PreLoadScreenBase.h b/Engine/Source/Runtime/PreLoadScreen/Public/PreLoadScreenBase.h
--- a/Engine/Source/Runtime/PreLoadScreen/Public/PreLoadScreenBase.h
+++ b/Engine/Source/Runtime/PreLoadScreen/Public/PreLoadScreenBase.h
@@ -14,41 +14,42 @@ class PRELOADSCREEN_API FPreLoadScreenBase :
     /**** IPreLoadScreen implementation ****/
 public:
     //We don't use these in the FPreLoadScreenBase, but they are useful for game-specific implementations
-    virtual void Tick(float DeltaTime) {}
-    virtual void RenderTick(float DeltaTime) {};
-    virtual void OnStop() {}
+    virtual void Tick(float DeltaTime) override {}
+	virtual bool ShouldRender() const override { return true; }
+    virtual void RenderTick(float DeltaTime) override {};
+    virtual void OnStop() override {}
 
     //Store off TargetWindow
-    virtual void OnPlay(TWeakPtr<SWindow> TargetWindow) { OwningWindow = TargetWindow; }
+    virtual void OnPlay(TWeakPtr<SWindow> TargetWindow) override { OwningWindow = TargetWindow; }
 
     //By default have a small added tick delay so we don't super spin out while waiting on other threads to load data / etc.
-    virtual float GetAddedTickDelay() { return 0.02f; }
+    virtual float GetAddedTickDelay() override { return 0.02f; }
 
-    virtual void Init();
+	virtual bool Init() override { return true; }
     
     virtual TSharedPtr<SWidget> GetWidget() override { return nullptr; }
     virtual const TSharedPtr<const SWidget> GetWidget() const override { return nullptr; }
 
     //IMPORTANT: This changes a LOT of functionality and implementation details. EarlyStartupScreens happen before the engine is fully initialized and block engine initialization before they finish.
     //           this means they have to forgo even the most basic of engine features like UObject support, as they are displayed before those systems are initialized.
-    virtual EPreLoadScreenTypes GetPreLoadScreenType() const { return EPreLoadScreenTypes::EngineLoadingScreen; }
+    virtual EPreLoadScreenTypes GetPreLoadScreenType() const override { return EPreLoadScreenTypes::EngineLoadingScreen; }
 
     virtual void SetEngineLoadingFinished(bool IsEngineLoadingFinished) override { bIsEngineLoadingFinished = IsEngineLoadingFinished; }
 
 	// PreLoadScreens not using this functionality should return NAME_None
 	virtual FName GetPreLoadScreenTag() const override { return NAME_None; }
 
-    virtual void CleanUp();
+    virtual void CleanUp() override;
 
     //Default behavior is just to see if we have an active widget. Should really overload with our own behavior to see if we are done displaying
-    virtual bool IsDone() const;
+    virtual bool IsDone() const override;
 
 public:
     FPreLoadScreenBase()
         : bIsEngineLoadingFinished(false)
     {}
 
-    virtual ~FPreLoadScreenBase() {};
+    virtual ~FPreLoadScreenBase() override {};
 
     //Handles constructing a FPreLoadSettingsContainerBase with the 
     virtual void InitSettingsFromConfig(const FString& ConfigFileName);
diff --git a/Engine/Source/Runtime/PreLoadScreen/Public/PreLoadScreenManager.h b/Engine/Source/Runtime/PreLoadScreen/Public/PreLoadScreenManager.h
--- a/Engine/Source/Runtime/PreLoadScreen/Public/PreLoadScreenManager.h
+++ b/Engine/Source/Runtime/PreLoadScreen/Public/PreLoadScreenManager.h
@@ -15,10 +15,14 @@
 
 #include "PreLoadSlateThreading.h"
 
+DECLARE_MULTICAST_DELEGATE_OneParam(FIsPreloadScreenResponsibleForRenderingMultiDelegate, bool);
+
 // Class that handles storing all registered PreLoadScreens and Playing/Stopping them
 class PRELOADSCREEN_API FPreLoadScreenManager
 {
 public:
+	static FIsPreloadScreenResponsibleForRenderingMultiDelegate IsResponsibleForRenderingDelegate;
+
     //Gets the single instance of this settings object. Also creates it if needed
     static FPreLoadScreenManager* Get()
     {
@@ -142,6 +146,8 @@ protected:
     static FCriticalSection RenderingEnabledCriticalSection;
     static bool bRenderingEnabled;
 
+	bool bIsResponsibleForRendering = false;
+
     double LastRenderTickTime;
 
     float OriginalSlateSleepVariableValue;
