Clean up comment

diff --git a/Engine/Source/Runtime/PreLoadScreen/Private/PreLoadScreenManager.cpp b/Engine/Source/Runtime/PreLoadScreen/Private/PreLoadScreenManager.cpp
--- a/Engine/Source/Runtime/PreLoadScreen/Private/PreLoadScreenManager.cpp
+++ b/Engine/Source/Runtime/PreLoadScreen/Private/PreLoadScreenManager.cpp
@@ -349,7 +349,6 @@ void FPreLoadScreenManager::EarlyPlayRenderFrameTick()
 
 	if (!ShouldRender())
 	{
-		// bIsResponsibleForRendering = true;
 		// In this case FPreLoadScreenManager is responsible for rendering but choosing not to, probably because the
 		// app is not in the fourground.
 		return;
