Hide startup logo on Switch when ESP is on

diff --git a/Engine/Source/Runtime/PreLoadScreen/Private/PreLoadScreenManager.cpp b/Engine/Source/Runtime/PreLoadScreen/Private/PreLoadScreenManager.cpp
--- a/Engine/Source/Runtime/PreLoadScreen/Private/PreLoadScreenManager.cpp
+++ b/Engine/Source/Runtime/PreLoadScreen/Private/PreLoadScreenManager.cpp
@@ -23,6 +23,10 @@
 #endif
 #endif
 
+#if PLATFORM_SWITCH
+#include <nn/oe.h>
+#endif
+
 IMPLEMENT_MODULE(FDefaultModuleImpl, PreLoadScreen);
 
 DEFINE_LOG_CATEGORY_STATIC(LogPreLoadScreenManager, Log, All);
@@ -155,6 +159,8 @@ void FPreLoadScreenManager::HandleEarlyStartupPlay()
 				bDidDisableScreensaver = FPlatformApplicationMisc::ControlScreensaver(FGenericPlatformApplicationMisc::EScreenSaverAction::Disable);
 			}
 
+			HidePlatformSpecificStartupScreen();
+
 			{
 				SCOPED_BOOT_TIMING("FPreLoadScreenManager::EarlyPlayFrameTick()");
 
@@ -331,6 +337,14 @@ void FPreLoadScreenManager::PlatformSpecificGameLogicFrameTick()
 #endif //PLATFORM_IOS
 }
 
+void FPreLoadScreenManager::HidePlatformSpecificStartupScreen()
+{
+#if PLATFORM_SWITCH
+	// hide the "Licensed by Nintendo" type logos, if the project had manual hide set
+	nn::oe::FinishStartupLogo();
+#endif // PLATFORM_SWITCH
+}
+
 bool FPreLoadScreenManager::ShouldRender()
 {
     FScopeLock ScopeLock(&RenderingEnabledCriticalSection);
diff --git a/Engine/Source/Runtime/PreLoadScreen/Public/PreLoadScreenManager.h b/Engine/Source/Runtime/PreLoadScreen/Public/PreLoadScreenManager.h
--- a/Engine/Source/Runtime/PreLoadScreen/Public/PreLoadScreenManager.h
+++ b/Engine/Source/Runtime/PreLoadScreen/Public/PreLoadScreenManager.h
@@ -117,6 +117,7 @@ protected:
     void EarlyPlayRenderFrameTick();
 
 	void PlatformSpecificGameLogicFrameTick();
+	void HidePlatformSpecificStartupScreen();
 
     /*** These functions describe how everything is handled during an non-Early PreLoadPlay. Everything is handled asynchronously in this case with a standalone renderer ***/
     void HandleEngineLoadingPlay();
