Hide startup logo on Switch when ESP is on
Move platform code to PlatformMisc

diff --git a/Engine/Source/Runtime/Core/Public/GenericPlatform/GenericPlatformMisc.h b/Engine/Source/Runtime/Core/Public/GenericPlatform/GenericPlatformMisc.h
--- a/Engine/Source/Runtime/Core/Public/GenericPlatform/GenericPlatformMisc.h
+++ b/Engine/Source/Runtime/Core/Public/GenericPlatform/GenericPlatformMisc.h
@@ -1371,6 +1371,10 @@ public:
 		return false;
 	}
 
+	static void HidePlatformStartupScreen()
+	{
+	}
+
 protected:
 	/** Whether the user should be prompted to allow for a remote debugger to be attached */
 	static bool bShouldPromptForRemoteDebugging;
diff --git a/Engine/Source/Runtime/PreLoadScreen/Private/PreLoadScreenManager.cpp b/Engine/Source/Runtime/PreLoadScreen/Private/PreLoadScreenManager.cpp
--- a/Engine/Source/Runtime/PreLoadScreen/Private/PreLoadScreenManager.cpp
+++ b/Engine/Source/Runtime/PreLoadScreen/Private/PreLoadScreenManager.cpp
@@ -23,10 +23,6 @@
 #endif
 #endif
 
-#if PLATFORM_SWITCH
-#include <nn/oe.h>
-#endif
-
 IMPLEMENT_MODULE(FDefaultModuleImpl, PreLoadScreen);
 
 DEFINE_LOG_CATEGORY_STATIC(LogPreLoadScreenManager, Log, All);
@@ -159,7 +155,7 @@ void FPreLoadScreenManager::HandleEarlyStartupPlay()
 				bDidDisableScreensaver = FPlatformApplicationMisc::ControlScreensaver(FGenericPlatformApplicationMisc::EScreenSaverAction::Disable);
 			}
 
-			HidePlatformSpecificStartupScreen();
+			FPlatformMisc::HidePlatformStartupScreen();
 
 			{
 				SCOPED_BOOT_TIMING("FPreLoadScreenManager::EarlyPlayFrameTick()");
@@ -337,14 +333,6 @@ void FPreLoadScreenManager::PlatformSpecificGameLogicFrameTick()
 #endif //PLATFORM_IOS
 }
 
-void FPreLoadScreenManager::HidePlatformSpecificStartupScreen()
-{
-#if PLATFORM_SWITCH
-	// hide the "Licensed by Nintendo" type logos, if the project had manual hide set
-	nn::oe::FinishStartupLogo();
-#endif // PLATFORM_SWITCH
-}
-
 bool FPreLoadScreenManager::ShouldRender()
 {
     FScopeLock ScopeLock(&RenderingEnabledCriticalSection);
diff --git a/Engine/Source/Runtime/PreLoadScreen/Public/PreLoadScreenManager.h b/Engine/Source/Runtime/PreLoadScreen/Public/PreLoadScreenManager.h
--- a/Engine/Source/Runtime/PreLoadScreen/Public/PreLoadScreenManager.h
+++ b/Engine/Source/Runtime/PreLoadScreen/Public/PreLoadScreenManager.h
@@ -117,7 +117,6 @@ protected:
     void EarlyPlayRenderFrameTick();
 
 	void PlatformSpecificGameLogicFrameTick();
-	void HidePlatformSpecificStartupScreen();
 
     /*** These functions describe how everything is handled during an non-Early PreLoadPlay. Everything is handled asynchronously in this case with a standalone renderer ***/
     void HandleEngineLoadingPlay();
