remove FAST_BOOT_HACK that is no longer needed

diff --git a/Engine/Source/Runtime/PreLoadScreen/Private/PreLoadScreenManager.cpp b/Engine/Source/Runtime/PreLoadScreen/Private/PreLoadScreenManager.cpp
--- a/Engine/Source/Runtime/PreLoadScreen/Private/PreLoadScreenManager.cpp
+++ b/Engine/Source/Runtime/PreLoadScreen/Private/PreLoadScreenManager.cpp
@@ -161,12 +161,7 @@ void FPreLoadScreenManager::HandleEarlyStartupPlay()
 				SCOPED_BOOT_TIMING("FPreLoadScreenManager::EarlyPlayFrameTick()");
 
 				//We run this PreLoadScreen until its finished or we lose the MainWindow as EarlyPreLoadPlay is synchronous
-#if BUILD_EMBEDDED_APP && FAST_BOOT_HACKS
-				FString ObjName(TEXT("LoggedInObject"));
-				while (FEmbeddedDelegates::GetNamedObject(ObjName) == nullptr || !PreLoadScreen->IsDone())
-#else
 				while (!PreLoadScreen->IsDone())
-#endif
 				{
 					EarlyPlayFrameTick();
 				}
