IInstallBundleManager:
Better handling of init error delegates.

PreloadScreen:
Allow an analytics factory to be passed in on Init.

diff --git a/Engine/Source/Runtime/InstallBundleManager/Public/InstallBundleManagerInterface.h b/Engine/Source/Runtime/InstallBundleManager/Public/InstallBundleManagerInterface.h
--- a/Engine/Source/Runtime/InstallBundleManager/Public/InstallBundleManagerInterface.h
+++ b/Engine/Source/Runtime/InstallBundleManager/Public/InstallBundleManagerInterface.h
@@ -305,8 +305,10 @@ public:
 
 	virtual bool HasBuildMetaData() const = 0;
 
-	virtual void PushInitErrorCallback(FInstallBundleManagerInitErrorHandler Callback) = 0;
+	virtual FDelegateHandle PushInitErrorCallback(FInstallBundleManagerInitErrorHandler Callback) = 0;
 	virtual void PopInitErrorCallback() = 0;
+	virtual void PopInitErrorCallback(FDelegateHandle Handle) = 0;
+	virtual void PopInitErrorCallback(const void* InUserObject) = 0;
 
 	virtual bool IsInitialized() const = 0;
 	virtual bool IsInitializing() const = 0;
diff --git a/Engine/Source/Runtime/NullInstallBundleManager/Source/NullInstallBundleManager.cpp b/Engine/Source/Runtime/NullInstallBundleManager/Source/NullInstallBundleManager.cpp
--- a/Engine/Source/Runtime/NullInstallBundleManager/Source/NullInstallBundleManager.cpp
+++ b/Engine/Source/Runtime/NullInstallBundleManager/Source/NullInstallBundleManager.cpp
@@ -11,14 +11,23 @@ class FNullInstallBundleManager : public IInstallBundleManager
 		return false;
 	}
 
-	virtual void PushInitErrorCallback(FInstallBundleManagerInitErrorHandler Callback) override
+	virtual FDelegateHandle PushInitErrorCallback(FInstallBundleManagerInitErrorHandler Callback) override
 	{
+		return FDelegateHandle();
 	}
 
 	virtual void PopInitErrorCallback() override
 	{
 	}
 
+	void PopInitErrorCallback(FDelegateHandle Handle) override
+	{
+	}
+
+	virtual void PopInitErrorCallback(const void* InUserObject) override
+	{
+	}
+
 	virtual bool IsInitializing() const override
 	{
 		return false;
diff --git a/Engine/Source/Runtime/PreLoadScreen/PreLoadScreen.Build.cs b/Engine/Source/Runtime/PreLoadScreen/PreLoadScreen.Build.cs
--- a/Engine/Source/Runtime/PreLoadScreen/PreLoadScreen.Build.cs
+++ b/Engine/Source/Runtime/PreLoadScreen/PreLoadScreen.Build.cs
@@ -4,41 +4,43 @@ using UnrealBuildTool;
 
 public class PreLoadScreen : ModuleRules
 {
-    public PreLoadScreen(ReadOnlyTargetRules Target) : base(Target)
-    {
-        PublicIncludePaths.Add("Runtime/PreLoadScreen/Public");
-        PrivateIncludePaths.Add("Runtime/PreLoadScreen/Private");
+	public PreLoadScreen(ReadOnlyTargetRules Target) : base(Target)
+	{
+		PublicIncludePaths.Add("Runtime/PreLoadScreen/Public");
+		PrivateIncludePaths.Add("Runtime/PreLoadScreen/Private");
 
-        PublicDependencyModuleNames.AddRange(
-            new string[] {
-                    "Engine",
-                    "ApplicationCore"
-                }
-        );
+		PublicDependencyModuleNames.AddRange(
+			new string[] {
+					"Engine",
+					"ApplicationCore",
+					"Analytics",
+					"AnalyticsET",
+				}
+		);
 
-        PrivateDependencyModuleNames.AddRange(
-            new string[] {
-                    "Core",
-                    "InputCore",
-                    "RenderCore",
-                    "CoreUObject",
-                    "RHI",
-                    "Slate",
-                    "SlateCore",
-                    "BuildPatchServices",
-                    "Projects"
-            }
-        );
+		PrivateDependencyModuleNames.AddRange(
+			new string[] {
+					"Core",
+					"InputCore",
+					"RenderCore",
+					"CoreUObject",
+					"RHI",
+					"Slate",
+					"SlateCore",
+					"BuildPatchServices",
+					"Projects",
+			}
+		);
 
 
-        //Need to make sure Android has Launch module so it can find and process AndroidEventManager events
-        if (Target.Platform == UnrealTargetPlatform.Android)
-        {
-            PrivateDependencyModuleNames.AddRange(
-                new string[] {
-                    "Launch"
-                }
-            );
-        }
-    }
+		//Need to make sure Android has Launch module so it can find and process AndroidEventManager events
+		if (Target.Platform == UnrealTargetPlatform.Android)
+		{
+			PrivateDependencyModuleNames.AddRange(
+				new string[] {
+					"Launch"
+				}
+			);
+		}
+	}
 }
diff --git a/Engine/Source/Runtime/PreLoadScreen/Public/PreLoadScreen.h b/Engine/Source/Runtime/PreLoadScreen/Public/PreLoadScreen.h
--- a/Engine/Source/Runtime/PreLoadScreen/Public/PreLoadScreen.h
+++ b/Engine/Source/Runtime/PreLoadScreen/Public/PreLoadScreen.h
@@ -4,6 +4,8 @@
 
 #include "Widgets/SWindow.h"
 
+#include "IAnalyticsProviderET.h"
+
 enum class EPreLoadScreenTypes : uint8
 {
 	CustomSplashScreen,
@@ -19,7 +21,7 @@ class IPreLoadScreen //: public TSharedFromThis<IPreLoadScreen, ESPMode::ThreadS
 public:
     virtual ~IPreLoadScreen() {};
 
-    virtual bool Init() = 0;
+    virtual bool Init(TFunctionRef<TSharedPtr<IAnalyticsProviderET>()> AnalyticsFactory) = 0;
     
     //Standard tick that happens every frame
     virtual void Tick(float DeltaTime) = 0;
diff --git a/Engine/Source/Runtime/PreLoadScreen/Public/PreLoadScreenBase.h b/Engine/Source/Runtime/PreLoadScreen/Public/PreLoadScreenBase.h
--- a/Engine/Source/Runtime/PreLoadScreen/Public/PreLoadScreenBase.h
+++ b/Engine/Source/Runtime/PreLoadScreen/Public/PreLoadScreenBase.h
@@ -25,7 +25,7 @@ public:
     //By default have a small added tick delay so we don't super spin out while waiting on other threads to load data / etc.
     virtual float GetAddedTickDelay() override { return 0.02f; }
 
-	virtual bool Init() override { return true; }
+	virtual bool Init(TFunctionRef<TSharedPtr<IAnalyticsProviderET>()> AnalyticsFactory) override { return true; }
     
     virtual TSharedPtr<SWidget> GetWidget() override { return nullptr; }
     virtual const TSharedPtr<const SWidget> GetWidget() const override { return nullptr; }
