If the engine needs to exit, while the slate loading screen is running, stop processing new slate frames, everything is going to start shuttingdown and it wont be safe to continue to make draw requests.

diff --git a/Engine/Source/Runtime/PreLoadScreen/Private/PreLoadSlateThreading.cpp b/Engine/Source/Runtime/PreLoadScreen/Private/PreLoadSlateThreading.cpp
--- a/Engine/Source/Runtime/PreLoadScreen/Private/PreLoadSlateThreading.cpp
+++ b/Engine/Source/Runtime/PreLoadScreen/Private/PreLoadSlateThreading.cpp
@@ -51,6 +51,13 @@ FPreLoadSlateWidgetRenderer::FPreLoadSlateWidgetRenderer(TSharedPtr<SWindow> InM
 
 void FPreLoadSlateWidgetRenderer::DrawWindow(float DeltaTime)
 {
+	// If the engine has requested exit, early out the draw loop, several
+	// of the things inside of here are not safe to perform while shutting down.
+	if (IsEngineExitRequested())
+	{
+		return;
+	}
+
     if (GDynamicRHI && GDynamicRHI->RHIIsRenderingSuspended())
     {
         // This avoids crashes if we Suspend rendering whilst the loading screen is up
