FPreLoadScreenManager::GameLogicFrameTick should call FTaskGraphInterface::Get().ProcessThreadUntilIdle(ENamedThreads::GameThread).  It's needed to properly finish Build Patch Services tasks.

diff --git a/Engine/Source/Runtime/PreLoadScreen/Private/PreLoadScreenManager.cpp b/Engine/Source/Runtime/PreLoadScreen/Private/PreLoadScreenManager.cpp
--- a/Engine/Source/Runtime/PreLoadScreen/Private/PreLoadScreenManager.cpp
+++ b/Engine/Source/Runtime/PreLoadScreen/Private/PreLoadScreenManager.cpp
@@ -339,6 +339,7 @@ void FPreLoadScreenManager::GameLogicFrameTick()
 		DeltaTime = FMath::Min(DeltaTime, MaxTickTime);
 
         //We have to manually tick everything as we are looping the main thread here
+		FTaskGraphInterface::Get().ProcessThreadUntilIdle(ENamedThreads::GameThread);
         FTicker::GetCoreTicker().Tick(DeltaTime);
         FThreadManager::Get().Tick();
 
