FPreLoadScreenManager:
Tick PlatformApplication.  This is needed on some platforms to dispatch system events.

diff --git a/Engine/Source/Runtime/PreLoadScreen/Private/PreLoadScreenManager.cpp b/Engine/Source/Runtime/PreLoadScreen/Private/PreLoadScreenManager.cpp
--- a/Engine/Source/Runtime/PreLoadScreen/Private/PreLoadScreenManager.cpp
+++ b/Engine/Source/Runtime/PreLoadScreen/Private/PreLoadScreenManager.cpp
@@ -356,6 +356,8 @@ void FPreLoadScreenManager::GameLogicFrameTick()
         // Gives widgets a chance to process any accumulated input
         FSlateApplication::Get().FinishedInputThisFrame();
 
+		FSlateApplication::Get().GetPlatformApplication()->Tick(DeltaTime);
+
         //Needed as this won't be incrementing on its own and some other tick functions rely on this (like analytics)
         GFrameCounter++;
     }
