Fix for crash when running clients and -nullrhi. Make sure the preload screen manager is running even though there is no window and we cannot render. There is logic in those systems which needs to execute even if we aren't rendering

diff --git a/Engine/Source/Runtime/PreLoadScreen/Private/PreLoadScreenManager.cpp b/Engine/Source/Runtime/PreLoadScreen/Private/PreLoadScreenManager.cpp
--- a/Engine/Source/Runtime/PreLoadScreen/Private/PreLoadScreenManager.cpp
+++ b/Engine/Source/Runtime/PreLoadScreen/Private/PreLoadScreenManager.cpp
@@ -142,16 +142,20 @@ void FPreLoadScreenManager::HandleEarlyStartupPlay()
     if (ensureAlwaysMsgf(HasActivePreLoadScreenType(EPreLoadScreenTypes::EarlyStartupScreen), TEXT("Invalid Active PreLoadScreen!")))
     {
         IPreLoadScreen* PreLoadScreen = GetActivePreLoadScreen();
-        if (PreLoadScreen && MainWindow.IsValid())
+        if (PreLoadScreen)
         {
 			SCOPED_BOOT_TIMING("FPreLoadScreenManager::HandleEarlyStartupPlay()");
 
-            PreLoadScreen->OnPlay(MainWindow.Pin());
+			PreLoadScreen->OnPlay(MainWindow);
 
-            if (PreLoadScreen->GetWidget().IsValid())
-            {
-                MainWindow.Pin()->SetContent(PreLoadScreen->GetWidget().ToSharedRef());
-            }
+			{
+				TSharedPtr<SWindow> MainWindowPtr = MainWindow.IsValid() ? MainWindow.Pin() : TSharedPtr<SWindow>();
+
+				if (MainWindowPtr.IsValid() && PreLoadScreen->GetWidget().IsValid())
+				{
+					MainWindowPtr->SetContent(PreLoadScreen->GetWidget().ToSharedRef());
+				}
+			}
 
 			bool bDidDisableScreensaver = false;
 			if (FPlatformApplicationMisc::IsScreensaverEnabled())
