Add missing FTaskScope on slate thread

diff --git a/Engine/Source/Runtime/PreLoadScreen/Private/PreLoadSlateThreading.cpp b/Engine/Source/Runtime/PreLoadScreen/Private/PreLoadSlateThreading.cpp
--- a/Engine/Source/Runtime/PreLoadScreen/Private/PreLoadSlateThreading.cpp
+++ b/Engine/Source/Runtime/PreLoadScreen/Private/PreLoadSlateThreading.cpp
@@ -189,6 +189,8 @@ bool FPreLoadScreenSlateSynchMechanism::IsSlateMainLoopRunning_AnyThread() const
 
 void FPreLoadScreenSlateSynchMechanism::RunMainLoop_SlateThread()
 {
+	FTaskTagScope Scope(ETaskTag::ESlateThread);
+
 	double LastTime = FPlatformTime::Seconds();
 	TAtomic<int32> SlateDrawEnqueuedCounter(0);
 	FEvent* EnqueueRenderEvent = FGenericPlatformProcess::GetSynchEventFromPool(false);
