Adding delegate that is called whenever FlushRenderingCommands is called.
Utilize new delegate to unlock AcquireCriticalSection in PreLoadScreenManager to prevent softlocks whenever FlushRenderingCommands is called during an EarlyStartup PreLoadScreen.
Remove previous hack in PreLoadScreenManager to unlock AcquireCriticalSectoin during fontcache flushes, as this was just to prevent the ensuing FlushRenderingCommands called by the font cache flush from soft locking and is now not needed.
#tests Test/Shipping Client

diff --git a/Engine/Source/Runtime/PreLoadScreen/Private/PreLoadScreenManager.cpp b/Engine/Source/Runtime/PreLoadScreen/Private/PreLoadScreenManager.cpp
--- a/Engine/Source/Runtime/PreLoadScreen/Private/PreLoadScreenManager.cpp
+++ b/Engine/Source/Runtime/PreLoadScreen/Private/PreLoadScreenManager.cpp
@@ -241,11 +241,7 @@ void FPreLoadScreenManager::HandleEarlyStartupPlay()
 
 			FPlatformMisc::HidePlatformStartupScreen();
 
-			//Register to handle case where font cache resources are released during rendering and can lead to a softlock
-			if (ensureAlwaysMsgf(FSlateApplication::Get().GetRenderer(), TEXT("Unexpected invalid SlateApplication Renderer!")))
-			{
-				FSlateApplication::Get().GetRenderer()->GetFontCache().Get().OnReleaseResources().AddRaw(this, &FPreLoadScreenManager::HandleReleaseFontResources);
-			}
+			RegisterDelegatesForEarlyStartupPlay();			
 
 			{
 				SCOPED_BOOT_TIMING("FPreLoadScreenManager::EarlyPlayFrameTick()");
@@ -257,10 +253,7 @@ void FPreLoadScreenManager::HandleEarlyStartupPlay()
 				}
 			}
 
-			if (FSlateApplication::Get().GetRenderer())
-			{
-				FSlateApplication::Get().GetRenderer()->GetFontCache().Get().OnReleaseResources().RemoveAll(this);
-			}
+			CleanUpDelegatesForEarlyStartupPlay();
 
 			if (bDidDisableScreensaver)
 			{
@@ -272,6 +265,19 @@ void FPreLoadScreenManager::HandleEarlyStartupPlay()
 	}
 }
 
+void FPreLoadScreenManager::RegisterDelegatesForEarlyStartupPlay()
+{
+	//Have to register to handle FlushRenderingCommands for the length of the PreLoadScreen
+	FCoreRenderDelegates::OnFlushRenderingCommandsStart.AddRaw(this, &FPreLoadScreenManager::HandleFlushRenderingCommandsStart);
+	FCoreRenderDelegates::OnFlushRenderingCommandsEnd.AddRaw(this, &FPreLoadScreenManager::HandleFlushRenderingCommandsEnd);
+}
+
+void FPreLoadScreenManager::CleanUpDelegatesForEarlyStartupPlay()
+{
+	FCoreRenderDelegates::OnFlushRenderingCommandsStart.RemoveAll(this);
+	FCoreRenderDelegates::OnFlushRenderingCommandsEnd.RemoveAll(this);
+}
+
 void FPreLoadScreenManager::HandleEngineLoadingPlay()
 {
 	if (ensureAlwaysMsgf(HasActivePreLoadScreenType(EPreLoadScreenTypes::EngineLoadingScreen), TEXT("Invalid Active PreLoadScreen!")))
@@ -689,15 +695,16 @@ bool FPreLoadScreenManager::ArePreLoadScreensEnabled()
 	return bEnabled;
 }
 
-void FPreLoadScreenManager::HandleReleaseFontResources(const class FSlateFontCache& InFontCache)
+void FPreLoadScreenManager::HandleFlushRenderingCommandsStart()
 {
-	UE_LOG(LogPreLoadScreenManager, Display, TEXT("Detected FontCache is invalidated. To prevent soft lock unlocking critical section and flushing rendering commands."));
-
-	//Hack fix for soft lock whenever font cache is invalidated. Invalidating our font cache during rendering will call this delegate before flushing rendering commands.
-	//This leads to us soft locking on trying to aquire our critical section during our rendering command. Here we pre-emptively unlock and flush rendering commands to clear
-	//this situation out pro-actively
+	//Whenever we flush rendering commands we need to unlock this critical section or we will softlock due to our enqueued rendering command never being able to
+	//acquire this critical section as its both locked on the GT and the GT won't unlock it as it's waiting for the rendering command to finish to unlock
 	AcquireCriticalSection.Unlock();
-	FlushRenderingCommands();
+}
+
+void FPreLoadScreenManager::HandleFlushRenderingCommandsEnd()
+{
+	//Relock critical section after the FlushRenderCommands finishes as we are
 	AcquireCriticalSection.Lock();
 }
 
diff --git a/Engine/Source/Runtime/PreLoadScreen/Public/PreLoadScreenManager.h b/Engine/Source/Runtime/PreLoadScreen/Public/PreLoadScreenManager.h
--- a/Engine/Source/Runtime/PreLoadScreen/Public/PreLoadScreenManager.h
+++ b/Engine/Source/Runtime/PreLoadScreen/Public/PreLoadScreenManager.h
@@ -106,7 +106,12 @@ protected:
 	void HandleStopPreLoadScreen();
 	void CleanUpResources();
 
-	void HandleReleaseFontResources(const class FSlateFontCache& InFontCache);
+	//Helpers that setup and clean-up delegates that only need to be active while we are playing an EarlyStartup PreLoadScreen
+	void RegisterDelegatesForEarlyStartupPlay();
+	void CleanUpDelegatesForEarlyStartupPlay();
+
+	void HandleFlushRenderingCommandsStart();
+	void HandleFlushRenderingCommandsEnd();
 
 	//Singleton Instance
 	struct FPreLoadScreenManagerDelete
diff --git a/Engine/Source/Runtime/RenderCore/Private/RenderingThread.cpp b/Engine/Source/Runtime/RenderCore/Private/RenderingThread.cpp
--- a/Engine/Source/Runtime/RenderCore/Private/RenderingThread.cpp
+++ b/Engine/Source/Runtime/RenderCore/Private/RenderingThread.cpp
@@ -32,6 +32,9 @@
 // Globals
 //
 
+FCoreRenderDelegates::FOnFlushRenderingCommandsStart FCoreRenderDelegates::OnFlushRenderingCommandsStart;
+FCoreRenderDelegates::FOnFlushRenderingCommandsEnd FCoreRenderDelegates::OnFlushRenderingCommandsEnd;
+
 UE_TRACE_CHANNEL_DEFINE(RenderCommandsChannel);
 
 RENDERCORE_API bool GIsThreadedRendering = false;
@@ -1281,6 +1284,7 @@ void FlushRenderingCommands()
 	}
 
 	TRACE_CPUPROFILER_EVENT_SCOPE(FlushRenderingCommands);
+	FCoreRenderDelegates::OnFlushRenderingCommandsStart.Broadcast();
 	FSuspendRenderingTickables SuspendRenderingTickables;
 
 	// Need to flush GT because render commands from threads other than GT are sent to
@@ -1310,6 +1314,8 @@ void FlushRenderingCommands()
 
 	// Delete the objects which were enqueued for deferred cleanup before the command queue flush.
 	delete PendingCleanupObjects;
+
+	FCoreRenderDelegates::OnFlushRenderingCommandsEnd.Broadcast();
 }
 
 void FlushPendingDeleteRHIResources_GameThread()
diff --git a/Engine/Source/Runtime/RenderCore/Public/RenderingThread.h b/Engine/Source/Runtime/RenderCore/Public/RenderingThread.h
--- a/Engine/Source/Runtime/RenderCore/Public/RenderingThread.h
+++ b/Engine/Source/Runtime/RenderCore/Public/RenderingThread.h
@@ -78,6 +78,15 @@ extern RENDERCORE_API void TickRenderingTickables();
 extern RENDERCORE_API void StartRenderCommandFenceBundler();
 extern RENDERCORE_API void StopRenderCommandFenceBundler();
 
+class RENDERCORE_API FCoreRenderDelegates
+{
+public:
+	DECLARE_MULTICAST_DELEGATE(FOnFlushRenderingCommandsStart);
+	static FOnFlushRenderingCommandsStart OnFlushRenderingCommandsStart;
+
+	DECLARE_MULTICAST_DELEGATE(FOnFlushRenderingCommandsEnd);
+	static FOnFlushRenderingCommandsEnd OnFlushRenderingCommandsEnd;
+};
 ////////////////////////////////////
 // Render thread suspension
 ////////////////////////////////////
