Rework how the PreLoadScreenSlateSynchMecanism wait for the renderthread. Used one resource intead of 2.

diff --git a/Engine/Source/Runtime/PreLoadScreen/Private/PreLoadSlateThreading.cpp b/Engine/Source/Runtime/PreLoadScreen/Private/PreLoadSlateThreading.cpp
--- a/Engine/Source/Runtime/PreLoadScreen/Private/PreLoadSlateThreading.cpp
+++ b/Engine/Source/Runtime/PreLoadScreen/Private/PreLoadSlateThreading.cpp
@@ -193,8 +193,8 @@ void FPreLoadScreenSlateSynchMechanism::RunMainLoop_SlateThread()
 	FTaskTagScope Scope(ETaskTag::ESlateThread);
 
 	double LastTime = FPlatformTime::Seconds();
-	TAtomic<int32> SlateDrawEnqueuedCounter(0);
-	FEvent* EnqueueRenderEvent = FGenericPlatformProcess::GetSynchEventFromPool(false);
+	bool bManualReset = false;
+	FEvent* EnqueueRenderEvent = FGenericPlatformProcess::GetSynchEventFromPool(bManualReset);
 
 	while (IsSlateMainLoopRunning_AnyThread())
 	{
@@ -218,14 +218,10 @@ void FPreLoadScreenSlateSynchMechanism::RunMainLoop_SlateThread()
 				CurrentTime = FPlatformTime::Seconds();
 				DeltaTime = CurrentTime - LastTime;
 			}
+			LastTime = CurrentTime;
 		}
 
-		// Wait for the current render command to be executed
-		while (IsSlateMainLoopRunning_AnyThread() && SlateDrawEnqueuedCounter != 0)
-		{
-			EnqueueRenderEvent->Wait(FTimespan::FromMilliseconds(1));
-		}
-
+		bool bRenderCommandEnqeued = false;
 		if (IsSlateMainLoopRunning_AnyThread())
 		{
 			// This avoids crashes if we Suspend rendering whilst the loading screen is up
@@ -236,18 +232,19 @@ void FPreLoadScreenSlateSynchMechanism::RunMainLoop_SlateThread()
 				if (IsSlateMainLoopRunning_AnyThread() && FPreLoadScreenManager::bRenderingEnabled)
 				{
 					WidgetRenderer->DrawWindow(DeltaTime);
-					++SlateDrawEnqueuedCounter;
+
+					bRenderCommandEnqeued = true;
 
 					//Queue up a render tick every time we tick on this sync thread.
 					FPreLoadScreenSlateSynchMechanism* SyncMech = this;
 					ENQUEUE_RENDER_COMMAND(PreLoadScreenRenderTick)(
-						[SyncMech, &SlateDrawEnqueuedCounter, EnqueueRenderEvent](FRHICommandListImmediate& RHICmdList)
+						[SyncMech, EnqueueRenderEvent](FRHICommandListImmediate& RHICmdList)
 						{
 							if (SyncMech->IsSlateMainLoopRunning_AnyThread())
 							{
 								FPreLoadScreenManager::StaticRenderTick_RenderThread();
 							}
-							--SlateDrawEnqueuedCounter;
+
 							EnqueueRenderEvent->Trigger();
 						}
 					);
@@ -255,13 +252,12 @@ void FPreLoadScreenSlateSynchMechanism::RunMainLoop_SlateThread()
 			}
 		}
 
-		LastTime = CurrentTime;
-	}
-
-	// Need to wait because the enqueued command has references this & SlateDrawEnqueuedCounter
-	while (SlateDrawEnqueuedCounter > 0)
-	{
-		EnqueueRenderEvent->Wait(FTimespan::FromMilliseconds(1));
+		if (bRenderCommandEnqeued)
+		{
+			// Release the lock and wait for the enqueued command to complete.
+			//Only one command at the time. Don't Flush because other commands might be less important and we are in the loading phase.
+			EnqueueRenderEvent->Wait();
+		}
 	}
 
 	FGenericPlatformProcess::ReturnSynchEventToPool(EnqueueRenderEvent);
