Prevents crash when the MoviePlayer and PreLoading system both tries to use the SlateThread at the same time. Add more info to track it down. The engine loading sequence already prevents PreLoading to launch when MoviePlayer is in used. Now it will ensure and recover if something else manually triggers both at the same times.
#prelfight 64c28f60f40ac84672005140

diff --git a/Engine/Source/Runtime/MoviePlayer/Private/MoviePlayerThreading.cpp b/Engine/Source/Runtime/MoviePlayer/Private/MoviePlayerThreading.cpp
--- a/Engine/Source/Runtime/MoviePlayer/Private/MoviePlayerThreading.cpp
+++ b/Engine/Source/Runtime/MoviePlayer/Private/MoviePlayerThreading.cpp
@@ -30,6 +30,7 @@ public:
 	/** FRunnable interface */
 	virtual bool Init() override;
 	virtual uint32 Run() override;
+	virtual void Exit() override;
 	virtual void Stop() override;
 private:
 	/** Hold a handle to our parent sync mechanism which handles all of our threading locks */
@@ -119,12 +120,24 @@ void FSlateLoadingSynchronizationMechanism::ResetSlateMainLoopRunning()
 	IsRunningSlateMainLoop.Reset();
 }
 
+void FSlateLoadingSynchronizationMechanism::SlateThreadInitFailed()
+{
+	bMainLoopRunning = false;
+}
+
 void FSlateLoadingSynchronizationMechanism::SlateThreadRunMainLoop()
 {
 	double LastTime = FPlatformTime::Seconds();
 
 	while (IsSlateMainLoopRunning())
 	{
+		// Test to ensure that we are still the SlateLoadingThread
+		checkCode(
+			const int32 CurrentThreadId = FPlatformTLS::GetCurrentThreadId();
+			const int32 PreviousThreadId = FPlatformAtomics::InterlockedCompareExchange((int32*)&GSlateLoadingThreadId, CurrentThreadId, CurrentThreadId);
+			check(PreviousThreadId == CurrentThreadId)
+		);
+
 		double CurrentTime = FPlatformTime::Seconds();
 		double DeltaTime = CurrentTime - LastTime;
 
@@ -179,25 +192,35 @@ bool FSlateLoadingThreadTask::Init()
 {
 	// First thing to do is set the slate loading thread ID
 	// This guarantees all systems know that a slate thread exists
-	GSlateLoadingThreadId = FPlatformTLS::GetCurrentThreadId();
+	const int32 PreviousValue = FPlatformAtomics::InterlockedCompareExchange((int32*)&GSlateLoadingThreadId, FPlatformTLS::GetCurrentThreadId(), 0);
 
-	return true;
+	bool bSuccess = PreviousValue == 0;
+	ensureMsgf(bSuccess, TEXT("Only one system can use the SlateThread at the same time. GetMoviePlayer is not compatible with PreLoadScreen."));
+	if (!bSuccess)
+	{
+		SyncMechanism->SlateThreadInitFailed();
+	}
+	return bSuccess;
 }
 
 uint32 FSlateLoadingThreadTask::Run()
 {
 	FTaskTagScope Scope(ETaskTag::ESlateThread);
-
 	check( GSlateLoadingThreadId == FPlatformTLS::GetCurrentThreadId() );
 
 	SyncMechanism->SlateThreadRunMainLoop();
 
-	// Tear down the slate loading thread ID
-	FPlatformAtomics::InterlockedExchange((int32*)&GSlateLoadingThreadId, 0);
-
 	return 0;
 }
 
+void FSlateLoadingThreadTask::Exit()
+{
+	// Tear down the slate loading thread ID
+	const int32 CurrentThreadId = FPlatformTLS::GetCurrentThreadId();
+	const int32 PreviousThreadId = FPlatformAtomics::InterlockedCompareExchange((int32*)&GSlateLoadingThreadId, 0, CurrentThreadId);
+	check(PreviousThreadId == CurrentThreadId);
+}
+
 void FSlateLoadingThreadTask::Stop()
 {
 	SyncMechanism->ResetSlateDrawPassEnqueued();
diff --git a/Engine/Source/Runtime/MoviePlayer/Private/MoviePlayerThreading.h b/Engine/Source/Runtime/MoviePlayer/Private/MoviePlayerThreading.h
--- a/Engine/Source/Runtime/MoviePlayer/Private/MoviePlayerThreading.h
+++ b/Engine/Source/Runtime/MoviePlayer/Private/MoviePlayerThreading.h
@@ -39,6 +39,8 @@ public:
 
 	/** The main loop to be run from the Slate thread */
 	void SlateThreadRunMainLoop();
+	/** The main loop could not be initialized */
+	void SlateThreadInitFailed();
 
 private:
 
diff --git a/Engine/Source/Runtime/PreLoadScreen/Private/PreLoadSlateThreading.cpp b/Engine/Source/Runtime/PreLoadScreen/Private/PreLoadSlateThreading.cpp
--- a/Engine/Source/Runtime/PreLoadScreen/Private/PreLoadSlateThreading.cpp
+++ b/Engine/Source/Runtime/PreLoadScreen/Private/PreLoadSlateThreading.cpp
@@ -21,18 +21,23 @@ TAtomic<int32> FPreLoadScreenSlateSynchMechanism::LoadingThreadInstanceCounter(0
 
 bool FPreLoadScreenSlateThreadTask::Init()
 {
-    // First thing to do is set the slate loading thread ID
-    // This guarantees all systems know that a slate thread exists
+	// First thing to do is set the slate loading thread ID
+	// This guarantees all systems know that a slate thread exists
 	const int32 PreviousValue = FPlatformAtomics::InterlockedCompareExchange((int32*)&GSlateLoadingThreadId, FPlatformTLS::GetCurrentThreadId(), 0);
-	check(PreviousValue == 0);
 
-    return true;
+	bool bSuccess = PreviousValue == 0;
+	ensureMsgf(bSuccess, TEXT("Only one system can use the SlateThread at the same time. PreLoadScreen is not compatible with GetMoviePlayer."));
+	return bSuccess;
 }
 
 uint32 FPreLoadScreenSlateThreadTask::Run()
 {
-    SyncMechanism->RunMainLoop_SlateThread();
-    return 0;
+	FTaskTagScope Scope(ETaskTag::ESlateThread);
+	check(GSlateLoadingThreadId == FPlatformTLS::GetCurrentThreadId());
+
+	SyncMechanism->RunMainLoop_SlateThread();
+
+	return 0;
 }
 
 void FPreLoadScreenSlateThreadTask::Exit()
@@ -44,9 +49,9 @@ void FPreLoadScreenSlateThreadTask::Exit()
 }
 
 FPreLoadSlateWidgetRenderer::FPreLoadSlateWidgetRenderer(TSharedPtr<SWindow> InMainWindow, TSharedPtr<SVirtualWindow> InVirtualRenderWindow, FSlateRenderer* InRenderer)
-    : MainWindow(InMainWindow.Get())
-    , VirtualRenderWindow(InVirtualRenderWindow.ToSharedRef())
-    , SlateRenderer(InRenderer)
+	: MainWindow(InMainWindow.Get())
+	, VirtualRenderWindow(InVirtualRenderWindow.ToSharedRef())
+	, SlateRenderer(InRenderer)
 {
     HittestGrid = MakeShareable(new FHittestGrid);
 }
@@ -155,21 +160,21 @@ void FPreLoadScreenSlateSynchMechanism::Initialize()
 
 void FPreLoadScreenSlateSynchMechanism::DestroySlateThread()
 {
-    check(IsInGameThread());
+	check(IsInGameThread());
 
-    if (bIsRunningSlateMainLoop)
-    {
+	if (bIsRunningSlateMainLoop)
+	{
 		check(SlateLoadingThread != nullptr);
 		bIsRunningSlateMainLoop = false;
 		SleepEvent->Trigger();
 		SlateLoadingThread->WaitForCompletion();
 
 		FGenericPlatformProcess::ReturnSynchEventToPool(SleepEvent);
-        delete SlateLoadingThread;
-        delete SlateRunnableTask;
-        SlateLoadingThread = nullptr;
-        SlateRunnableTask = nullptr;
-    }
+		delete SlateLoadingThread;
+		delete SlateRunnableTask;
+		SlateLoadingThread = nullptr;
+		SlateRunnableTask = nullptr;
+	}
 }
 
 void FPreLoadScreenSlateSynchMechanism::HandleWindowBeingDestroyed(const SWindow& WindowBeingDestroyed)
@@ -190,8 +195,6 @@ bool FPreLoadScreenSlateSynchMechanism::IsSlateMainLoopRunning_AnyThread() const
 
 void FPreLoadScreenSlateSynchMechanism::RunMainLoop_SlateThread()
 {
-	FTaskTagScope Scope(ETaskTag::ESlateThread);
-
 	double LastTime = FPlatformTime::Seconds();
 	bool bManualReset = false;
 	FEvent* EnqueueRenderEvent = FGenericPlatformProcess::GetSynchEventFromPool(bManualReset);
