Speculative fix PreloadScreen deadlock.
---
Keep track on if we are locked similar to slate main rendering, also keep track of locking thread (Not used yet)
Don't lock on flush end if we weren't locking before
Surround 2x main render commands w/ 'UE::RenderCommandPipe::FSyncScope' (To serialize commands in addition to manual sync call).
Minor: Replace existing manual unlock -> lock with 'FScopeUnlock'

diff --git a/Engine/Source/Runtime/PreLoadScreen/Private/PreLoadScreenManager.cpp b/Engine/Source/Runtime/PreLoadScreen/Private/PreLoadScreenManager.cpp
--- a/Engine/Source/Runtime/PreLoadScreen/Private/PreLoadScreenManager.cpp
+++ b/Engine/Source/Runtime/PreLoadScreen/Private/PreLoadScreenManager.cpp
@@ -12,6 +12,7 @@
 
 #include "HAL/ThreadManager.h"
 #include "Modules/ModuleManager.h"
+#include "RenderingThread.h"
 
 #if BUILD_EMBEDDED_APP
 #include "Misc/EmbeddedCommunication.h"
@@ -30,6 +31,9 @@ DEFINE_LOG_CATEGORY_STATIC(LogPreLoadScreenManager, Log, All);
 TUniquePtr<FPreLoadScreenManager, FPreLoadScreenManager::FPreLoadScreenManagerDelete> FPreLoadScreenManager::Instance;
 FCriticalSection FPreLoadScreenManager::AcquireCriticalSection;
 TAtomic<bool> FPreLoadScreenManager::bRenderingEnabled(true);
+TAtomic<bool> FPreLoadScreenManager::bIsLocked(false);
+TAtomic<bool> FPreLoadScreenManager::bIsLockedByGameThread(false);
+TAtomic<bool> FPreLoadScreenManager::bIsLockedByRenderThread(false);
 
 FPreLoadScreenManager* FPreLoadScreenManager::Get()
 {
@@ -42,6 +46,8 @@ void FPreLoadScreenManager::Create()
 
 	// Lock in case a user decide to create/detroy the manager multiple times
 	FScopeLock Lock(&AcquireCriticalSection);
+	TGuardValue LockedGuard(bIsLocked, true);
+	TGuardValue LockedGuardGameThread(bIsLockedByGameThread, true);
 
 	if (!Instance.IsValid() && ArePreLoadScreensEnabled())
 	{
@@ -54,6 +60,8 @@ void FPreLoadScreenManager::Destroy()
 	check(IsInGameThread());
 
 	FScopeLock Lock(&AcquireCriticalSection); // Make sure the render thread is completed before cleaning up
+	TGuardValue LockedGuard(bIsLocked, true);
+	TGuardValue LockedGuardGameThread(bIsLockedByGameThread, true);
 
 	if (Instance.IsValid())
 	{
@@ -124,6 +132,9 @@ void FPreLoadScreenManager::UnRegisterPreLoadScreen(const TSharedPtr<IPreLoadScr
 	check(IsInGameThread());
 
 	FScopeLock Lock(&AcquireCriticalSection);
+	TGuardValue LockedGuard(bIsLocked, true);
+	TGuardValue LockedGuardGameThread(bIsLockedByGameThread, true);
+
 	if (PreLoadScreen.IsValid())
 	{
 		const int32 IndexOf = PreLoadScreens.IndexOfByKey(PreLoadScreen);
@@ -169,6 +180,8 @@ void FPreLoadScreenManager::PlayPreLoadScreenAtIndex(int32 Index)
 	if (ArePreLoadScreensEnabled())
 	{
 		FScopeLock Lock(&AcquireCriticalSection);
+		TGuardValue LockedGuard(bIsLocked, true);
+		TGuardValue LockedGuardGameThread(bIsLockedByGameThread, true);
 
 		if (ensureAlwaysMsgf(!HasValidActivePreLoadScreen(), TEXT("Call to FPreLoadScreenManager::PlayPreLoadScreenAtIndex when something is already playing.")))
 		{
@@ -357,6 +370,9 @@ void FPreLoadScreenManager::StaticRenderTick_RenderThread()
 	check(IsInRenderingThread());
 
 	FScopeLock Lock(&AcquireCriticalSection);
+	TGuardValue LockedGuard(bIsLocked, true);
+	TGuardValue LockedGuardGameThread(bIsLockedByRenderThread, true);
+
 	if (ensure(FPreLoadScreenManager::Get())) // The manager should clear the slate render thread before closing
 	{
 		FPreLoadScreenManager::Get()->RenderTick_RenderThread();
@@ -495,7 +511,11 @@ void FPreLoadScreenManager::PlatformSpecificGameLogicFrameTick()
 
 void FPreLoadScreenManager::EnableRendering(bool bEnabled)
 {
-    FScopeLock ScopeLock(&AcquireCriticalSection);
+	FScopeLock ScopeLock(&AcquireCriticalSection);
+	TGuardValue LockedGuard(bIsLocked, true);
+	TGuardValue LockedGuardGameThread(bIsLockedByGameThread, IsInGameThread());
+	TGuardValue LockedGuardRenderThread(bIsLockedByGameThread, IsInRenderingThread());
+
     bRenderingEnabled = bEnabled;
 }
 
@@ -509,9 +529,13 @@ void FPreLoadScreenManager::EarlyPlayRenderFrameTick()
 		// app is not in the foreground.
 
 		// Cycle lock to give a chance to another thread to re-enable rendering.
-		AcquireCriticalSection.Unlock();
-		FPlatformProcess::Sleep(0);
-		AcquireCriticalSection.Lock();
+		{
+			FScopeUnlock ScopeUnlock(&AcquireCriticalSection);
+			TGuardValue LockedGuard(bIsLocked, false);
+			TGuardValue LockedGuardGameThread(bIsLockedByGameThread, false);
+			TGuardValue LockedGuardRenderThread(bIsLockedByGameThread, false);
+			FPlatformProcess::Sleep(0);
+		}
 		return;
 	}
 
@@ -533,42 +557,55 @@ void FPreLoadScreenManager::EarlyPlayRenderFrameTick()
 		{
 			FSlateApplication& SlateApp = FSlateApplication::Get();
 			float SlateDeltaTime = SlateApp.GetDeltaTime();
-
-			//Setup Slate Render Command
 			FPreLoadScreenManager* Self = this;
-			ENQUEUE_RENDER_COMMAND(BeginPreLoadScreenFrame)(
-				[Self, SlateDeltaTime](FRHICommandListImmediate& RHICmdList)
-				{
-					FScopeLock Lock(&FPreLoadScreenManager::AcquireCriticalSection);
-					// Self is still valid because we do a FlushRenderingCommands in StopPreLoadScreen
-					if (Self->bRenderingEnabled && Self->HasActivePreLoadScreenTypeForEarlyStartup() && !Self->bHasRenderPreLoadScreenFrame_RenderThread)
-					{
-						GFrameNumberRenderThread++;
-						GRHICommandList.GetImmediateCommandList().BeginFrame();
 
-						Self->bHasRenderPreLoadScreenFrame_RenderThread = true;
-						IPreLoadScreen* ActivePreLoadScreen = Self->PreLoadScreens[Self->ActivePreLoadScreenIndex].Get();
-						ActivePreLoadScreen->RenderTick(SlateDeltaTime);
-					}
-				});
-
-			AcquireCriticalSection.Unlock();
+			{
+				UE::RenderCommandPipe::FSyncScope SyncScope;
 
-			SlateApp.Tick();
+				//Setup Slate Render Command
+				ENQUEUE_RENDER_COMMAND(BeginPreLoadScreenFrame)(
+					[Self, SlateDeltaTime](FRHICommandListImmediate& RHICmdList)
+					{
+						FScopeLock Lock(&FPreLoadScreenManager::AcquireCriticalSection);
+						TGuardValue LockedGuard(bIsLocked, true);
+						TGuardValue LockedGuardRenderThread(bIsLockedByGameThread, true);
+
+						// Self is still valid because we do a FlushRenderingCommands in StopPreLoadScreen
+						if (Self->bRenderingEnabled && Self->HasActivePreLoadScreenTypeForEarlyStartup() && !Self->bHasRenderPreLoadScreenFrame_RenderThread)
+						{
+							GFrameNumberRenderThread++;
+							GRHICommandList.GetImmediateCommandList().BeginFrame();
+
+							Self->bHasRenderPreLoadScreenFrame_RenderThread = true;
+							IPreLoadScreen* ActivePreLoadScreen = Self->PreLoadScreens[Self->ActivePreLoadScreenIndex].Get();
+							ActivePreLoadScreen->RenderTick(SlateDeltaTime);
+						}
+					});
+			}
 
-			// Synchronize the game thread and the render thread so that the render thread doesn't get too far behind.
-			SlateApp.GetRenderer()->Sync();
+			{
+				FScopeUnlock ScopeUnlock(&AcquireCriticalSection);
+				TGuardValue LockedGuard(bIsLocked, false);
+				TGuardValue LockedGuardGameThread(bIsLockedByGameThread, false);
+				TGuardValue LockedGuardRenderThread(bIsLockedByGameThread, false);
+				SlateApp.Tick();
+
+				// Synchronize the game thread and the render thread so that the render thread doesn't get too far behind.
+				SlateApp.GetRenderer()->Sync();
+			}
 
-			AcquireCriticalSection.Lock();
+			{
+				UE::RenderCommandPipe::FSyncScope SyncScope;
 
-			ENQUEUE_RENDER_COMMAND(FinishPreLoadScreenFrame)(
-				[Self](FRHICommandListImmediate& RHICmdList)
-				{
-					// Self is still valid because we do a FlushRenderingCommands in StopPreLoadScreen
-					Self->bHasRenderPreLoadScreenFrame_RenderThread = false;
-					GRHICommandList.GetImmediateCommandList().EndFrame();
-					GRHICommandList.GetImmediateCommandList().ImmediateFlush(EImmediateFlushType::FlushRHIThreadFlushResources);
-				});
+				ENQUEUE_RENDER_COMMAND(FinishPreLoadScreenFrame)(
+					[Self](FRHICommandListImmediate& RHICmdList)
+					{
+						// Self is still valid because we do a FlushRenderingCommands in StopPreLoadScreen
+						Self->bHasRenderPreLoadScreenFrame_RenderThread = false;
+						GRHICommandList.GetImmediateCommandList().EndFrame();
+						GRHICommandList.GetImmediateCommandList().ImmediateFlush(EImmediateFlushType::FlushRHIThreadFlushResources);
+					});
+			}
 		}
 	}
 }
@@ -584,9 +621,13 @@ void FPreLoadScreenManager::StopPreLoadScreen()
 			HandleStopPreLoadScreen();
 		}
 
-		AcquireCriticalSection.Unlock();
-		FlushRenderingCommands();
-		AcquireCriticalSection.Lock();
+		{
+			FScopeUnlock ScopeUnlock(&AcquireCriticalSection);
+			TGuardValue LockedGuard(bIsLocked, false);
+			TGuardValue LockedGuardGameThread(bIsLockedByGameThread, false);
+			TGuardValue LockedGuardRenderThread(bIsLockedByGameThread, false);
+			FlushRenderingCommands();
+		}
 	}
 }
 
@@ -594,6 +635,9 @@ void FPreLoadScreenManager::HandleStopPreLoadScreen()
 {
 	{
 		FScopeLock Lock(&AcquireCriticalSection); // prevent stop while we are rendering the preloadscreen
+		TGuardValue LockedGuard(bIsLocked, true);
+		TGuardValue LockedGuardGameThread(bIsLockedByGameThread, IsInGameThread());
+		TGuardValue LockedGuardRenderThread(bIsLockedByGameThread, IsInRenderingThread());
 
 		if (HasValidActivePreLoadScreen())
 		{
@@ -619,6 +663,9 @@ void FPreLoadScreenManager::PassPreLoadScreenWindowBackToGame() const
     if (IsUsingMainWindow())
     {
 		FScopeLock Lock(&AcquireCriticalSection); // wait until we finish with rendering before moving the context
+		TGuardValue LockedGuard(bIsLocked, true);
+		TGuardValue LockedGuardGameThread(bIsLockedByGameThread, IsInGameThread());
+		TGuardValue LockedGuardRenderThread(bIsLockedByGameThread, IsInRenderingThread());
 
         UGameEngine* GameEngine = Cast<UGameEngine>(GEngine);
         if (MainWindow.IsValid() && GameEngine)
@@ -701,12 +748,20 @@ void FPreLoadScreenManager::HandleFlushRenderingCommandsStart()
 	//Whenever we flush rendering commands we need to unlock this critical section or we will softlock due to our enqueued rendering command never being able to
 	//acquire this critical section as its both locked on the GT and the GT won't unlock it as it's waiting for the rendering command to finish to unlock
 	AcquireCriticalSection.Unlock();
+
+	// Note: Don't modify bIsLocked as that will determine if we relock when flush completes
 }
 
 void FPreLoadScreenManager::HandleFlushRenderingCommandsEnd()
 {
 	//Relock critical section after the FlushRenderCommands finishes as we are
 	AcquireCriticalSection.Lock();
+
+	// Check if we were locked before, if not release lock to prevent flush from leaving us in undesired locked stated
+	if (!bIsLocked)
+	{
+		AcquireCriticalSection.Unlock();
+	}
 }
 
 void FPreLoadScreenManager::CleanUpResources()
diff --git a/Engine/Source/Runtime/PreLoadScreen/Public/PreLoadScreenManager.h b/Engine/Source/Runtime/PreLoadScreen/Public/PreLoadScreenManager.h
--- a/Engine/Source/Runtime/PreLoadScreen/Public/PreLoadScreenManager.h
+++ b/Engine/Source/Runtime/PreLoadScreen/Public/PreLoadScreenManager.h
@@ -125,6 +125,9 @@ protected:
 
 	static PRELOADSCREEN_API FCriticalSection AcquireCriticalSection;
 	static PRELOADSCREEN_API TAtomic<bool> bRenderingEnabled;
+	static PRELOADSCREEN_API TAtomic<bool> bIsLocked;
+	static PRELOADSCREEN_API TAtomic<bool> bIsLockedByGameThread;
+	static PRELOADSCREEN_API TAtomic<bool> bIsLockedByRenderThread;
 
 	TArray<TSharedPtr<IPreLoadScreen>> PreLoadScreens;
 
