Clean up Preload Screen Manager locking and render commands

#localization none
#tests Tested on Windows, Lorry smoked on Switch

diff --git a/Engine/Source/Runtime/PreLoadScreen/Private/PreLoadScreenManager.cpp b/Engine/Source/Runtime/PreLoadScreen/Private/PreLoadScreenManager.cpp
--- a/Engine/Source/Runtime/PreLoadScreen/Private/PreLoadScreenManager.cpp
+++ b/Engine/Source/Runtime/PreLoadScreen/Private/PreLoadScreenManager.cpp
@@ -12,7 +12,6 @@
 
 #include "HAL/ThreadManager.h"
 #include "Modules/ModuleManager.h"
-#include "RenderingThread.h"
 
 #if BUILD_EMBEDDED_APP
 #include "Misc/EmbeddedCommunication.h"
@@ -30,10 +29,7 @@ DEFINE_LOG_CATEGORY_STATIC(LogPreLoadScreenManager, Log, All);
 
 TUniquePtr<FPreLoadScreenManager, FPreLoadScreenManager::FPreLoadScreenManagerDelete> FPreLoadScreenManager::Instance;
 FCriticalSection FPreLoadScreenManager::AcquireCriticalSection;
-TAtomic<bool> FPreLoadScreenManager::bRenderingEnabled(true);
-TAtomic<bool> FPreLoadScreenManager::bIsLocked(false);
-TAtomic<bool> FPreLoadScreenManager::bIsLockedByGameThread(false);
-TAtomic<bool> FPreLoadScreenManager::bIsLockedByRenderThread(false);
+std::atomic<bool> FPreLoadScreenManager::bRenderingEnabled(true);
 
 FPreLoadScreenManager* FPreLoadScreenManager::Get()
 {
@@ -46,8 +42,6 @@ void FPreLoadScreenManager::Create()
 
 	// Lock in case a user decide to create/detroy the manager multiple times
 	FScopeLock Lock(&AcquireCriticalSection);
-	TGuardValue LockedGuard(bIsLocked, true);
-	TGuardValue LockedGuardGameThread(bIsLockedByGameThread, true);
 
 	if (!Instance.IsValid() && ArePreLoadScreensEnabled())
 	{
@@ -60,8 +54,6 @@ void FPreLoadScreenManager::Destroy()
 	check(IsInGameThread());
 
 	FScopeLock Lock(&AcquireCriticalSection); // Make sure the render thread is completed before cleaning up
-	TGuardValue LockedGuard(bIsLocked, true);
-	TGuardValue LockedGuardGameThread(bIsLockedByGameThread, true);
 
 	if (Instance.IsValid())
 	{
@@ -132,8 +124,6 @@ void FPreLoadScreenManager::UnRegisterPreLoadScreen(const TSharedPtr<IPreLoadScr
 	check(IsInGameThread());
 
 	FScopeLock Lock(&AcquireCriticalSection);
-	TGuardValue LockedGuard(bIsLocked, true);
-	TGuardValue LockedGuardGameThread(bIsLockedByGameThread, true);
 
 	if (PreLoadScreen.IsValid())
 	{
@@ -180,8 +170,6 @@ void FPreLoadScreenManager::PlayPreLoadScreenAtIndex(int32 Index)
 	if (ArePreLoadScreensEnabled())
 	{
 		FScopeLock Lock(&AcquireCriticalSection);
-		TGuardValue LockedGuard(bIsLocked, true);
-		TGuardValue LockedGuardGameThread(bIsLockedByGameThread, true);
 
 		if (ensureAlwaysMsgf(!HasValidActivePreLoadScreen(), TEXT("Call to FPreLoadScreenManager::PlayPreLoadScreenAtIndex when something is already playing.")))
 		{
@@ -254,8 +242,6 @@ void FPreLoadScreenManager::HandleEarlyStartupPlay()
 
 			FPlatformMisc::HidePlatformStartupScreen();
 
-			RegisterDelegatesForEarlyStartupPlay();			
-
 			{
 				SCOPED_BOOT_TIMING("FPreLoadScreenManager::EarlyPlayFrameTick()");
 
@@ -266,8 +252,6 @@ void FPreLoadScreenManager::HandleEarlyStartupPlay()
 				}
 			}
 
-			CleanUpDelegatesForEarlyStartupPlay();
-
 			if (bDidDisableScreensaver)
 			{
 				FPlatformApplicationMisc::ControlScreensaver(FGenericPlatformApplicationMisc::EScreenSaverAction::Enable);
@@ -278,19 +262,6 @@ void FPreLoadScreenManager::HandleEarlyStartupPlay()
 	}
 }
 
-void FPreLoadScreenManager::RegisterDelegatesForEarlyStartupPlay()
-{
-	//Have to register to handle FlushRenderingCommands for the length of the PreLoadScreen
-	FCoreRenderDelegates::OnFlushRenderingCommandsStart.AddRaw(this, &FPreLoadScreenManager::HandleFlushRenderingCommandsStart);
-	FCoreRenderDelegates::OnFlushRenderingCommandsEnd.AddRaw(this, &FPreLoadScreenManager::HandleFlushRenderingCommandsEnd);
-}
-
-void FPreLoadScreenManager::CleanUpDelegatesForEarlyStartupPlay()
-{
-	FCoreRenderDelegates::OnFlushRenderingCommandsStart.RemoveAll(this);
-	FCoreRenderDelegates::OnFlushRenderingCommandsEnd.RemoveAll(this);
-}
-
 void FPreLoadScreenManager::HandleEngineLoadingPlay()
 {
 	if (ensureAlwaysMsgf(HasActivePreLoadScreenType(EPreLoadScreenTypes::EngineLoadingScreen), TEXT("Invalid Active PreLoadScreen!")))
@@ -370,8 +341,6 @@ void FPreLoadScreenManager::StaticRenderTick_RenderThread()
 	check(IsInRenderingThread());
 
 	FScopeLock Lock(&AcquireCriticalSection);
-	TGuardValue LockedGuard(bIsLocked, true);
-	TGuardValue LockedGuardGameThread(bIsLockedByRenderThread, true);
 
 	if (ensure(FPreLoadScreenManager::Get())) // The manager should clear the slate render thread before closing
 	{
@@ -511,37 +480,15 @@ void FPreLoadScreenManager::PlatformSpecificGameLogicFrameTick()
 
 void FPreLoadScreenManager::EnableRendering(bool bEnabled)
 {
-	FScopeLock ScopeLock(&AcquireCriticalSection);
-	TGuardValue LockedGuard(bIsLocked, true);
-	TGuardValue LockedGuardGameThread(bIsLockedByGameThread, IsInGameThread());
-	TGuardValue LockedGuardRenderThread(bIsLockedByGameThread, IsInRenderingThread());
-
     bRenderingEnabled = bEnabled;
 }
 
 void FPreLoadScreenManager::EarlyPlayRenderFrameTick()
 {
-	bool bIsResponsibleForRendering_Local = true;
-
-	if (!bRenderingEnabled || !FSlateApplication::IsInitialized())
-	{
-		// If rendering disabled, FPreLoadScreenManager is responsible for rendering but choosing not to, probably because the
-		// app is not in the foreground.
-
-		// Cycle lock to give a chance to another thread to re-enable rendering.
-		{
-			FScopeUnlock ScopeUnlock(&AcquireCriticalSection);
-			TGuardValue LockedGuard(bIsLocked, false);
-			TGuardValue LockedGuardGameThread(bIsLockedByGameThread, false);
-			TGuardValue LockedGuardRenderThread(bIsLockedByGameThread, false);
-			FPlatformProcess::Sleep(0);
-		}
-		return;
-	}
-
 	IPreLoadScreen* ActivePreLoadScreen = PreLoadScreens[ActivePreLoadScreenIndex].Get();
 	if (ensureAlwaysMsgf(ActivePreLoadScreen, TEXT("Invalid Active PreLoadScreen during EarlyPlayRenderFrameTick!")))
 	{
+		bool bIsResponsibleForRendering_Local = true;
 		if (!ActivePreLoadScreen->ShouldRender())
 		{
 			bIsResponsibleForRendering_Local = false;
@@ -556,56 +503,10 @@ void FPreLoadScreenManager::EarlyPlayRenderFrameTick()
 		if (bIsResponsibleForRendering_Local)
 		{
 			FSlateApplication& SlateApp = FSlateApplication::Get();
-			float SlateDeltaTime = SlateApp.GetDeltaTime();
-			FPreLoadScreenManager* Self = this;
-
-			{
-				UE::RenderCommandPipe::FSyncScope SyncScope;
-
-				//Setup Slate Render Command
-				ENQUEUE_RENDER_COMMAND(BeginPreLoadScreenFrame)(
-					[Self, SlateDeltaTime](FRHICommandListImmediate& RHICmdList)
-					{
-						FScopeLock Lock(&FPreLoadScreenManager::AcquireCriticalSection);
-						TGuardValue LockedGuard(bIsLocked, true);
-						TGuardValue LockedGuardRenderThread(bIsLockedByGameThread, true);
-
-						// Self is still valid because we do a FlushRenderingCommands in StopPreLoadScreen
-						if (Self->bRenderingEnabled && Self->HasActivePreLoadScreenTypeForEarlyStartup() && !Self->bHasRenderPreLoadScreenFrame_RenderThread)
-						{
-							GFrameNumberRenderThread++;
-							GRHICommandList.GetImmediateCommandList().BeginFrame();
-
-							Self->bHasRenderPreLoadScreenFrame_RenderThread = true;
-							IPreLoadScreen* ActivePreLoadScreen = Self->PreLoadScreens[Self->ActivePreLoadScreenIndex].Get();
-							ActivePreLoadScreen->RenderTick(SlateDeltaTime);
-						}
-					});
-			}
-
-			{
-				FScopeUnlock ScopeUnlock(&AcquireCriticalSection);
-				TGuardValue LockedGuard(bIsLocked, false);
-				TGuardValue LockedGuardGameThread(bIsLockedByGameThread, false);
-				TGuardValue LockedGuardRenderThread(bIsLockedByGameThread, false);
-				SlateApp.Tick();
-
-				// Synchronize the game thread and the render thread so that the render thread doesn't get too far behind.
-				SlateApp.GetRenderer()->Sync();
-			}
+			SlateApp.Tick();
 
-			{
-				UE::RenderCommandPipe::FSyncScope SyncScope;
-
-				ENQUEUE_RENDER_COMMAND(FinishPreLoadScreenFrame)(
-					[Self](FRHICommandListImmediate& RHICmdList)
-					{
-						// Self is still valid because we do a FlushRenderingCommands in StopPreLoadScreen
-						Self->bHasRenderPreLoadScreenFrame_RenderThread = false;
-						GRHICommandList.GetImmediateCommandList().EndFrame();
-						GRHICommandList.GetImmediateCommandList().ImmediateFlush(EImmediateFlushType::FlushRHIThreadFlushResources);
-					});
-			}
+			// Synchronize the game thread and the render thread so that the render thread doesn't get too far behind.
+			SlateApp.GetRenderer()->Sync();
 		}
 	}
 }
@@ -622,10 +523,7 @@ void FPreLoadScreenManager::StopPreLoadScreen()
 		}
 
 		{
-			FScopeUnlock ScopeUnlock(&AcquireCriticalSection);
-			TGuardValue LockedGuard(bIsLocked, false);
-			TGuardValue LockedGuardGameThread(bIsLockedByGameThread, false);
-			TGuardValue LockedGuardRenderThread(bIsLockedByGameThread, false);
+			FScopeUnlock ScopeUnlock(&AcquireCriticalSection); // Must unlock, StaticRenderTick_RenderThread will acquire the lock
 			FlushRenderingCommands();
 		}
 	}
@@ -635,9 +533,6 @@ void FPreLoadScreenManager::HandleStopPreLoadScreen()
 {
 	{
 		FScopeLock Lock(&AcquireCriticalSection); // prevent stop while we are rendering the preloadscreen
-		TGuardValue LockedGuard(bIsLocked, true);
-		TGuardValue LockedGuardGameThread(bIsLockedByGameThread, IsInGameThread());
-		TGuardValue LockedGuardRenderThread(bIsLockedByGameThread, IsInRenderingThread());
 
 		if (HasValidActivePreLoadScreen())
 		{
@@ -660,33 +555,30 @@ void FPreLoadScreenManager::HandleStopPreLoadScreen()
 
 void FPreLoadScreenManager::PassPreLoadScreenWindowBackToGame() const
 {
-    if (IsUsingMainWindow())
-    {
+	if (IsUsingMainWindow())
+	{
 		FScopeLock Lock(&AcquireCriticalSection); // wait until we finish with rendering before moving the context
-		TGuardValue LockedGuard(bIsLocked, true);
-		TGuardValue LockedGuardGameThread(bIsLockedByGameThread, IsInGameThread());
-		TGuardValue LockedGuardRenderThread(bIsLockedByGameThread, IsInRenderingThread());
 
-        UGameEngine* GameEngine = Cast<UGameEngine>(GEngine);
-        if (MainWindow.IsValid() && GameEngine)
-        {
-            GameEngine->GameViewportWindow = MainWindow;
-        }
-        else
-        {
-            UE_LOG(LogPreLoadScreenManager, Warning, TEXT("FPreLoadScreenManager::PassLoadingScreenWindowBackToGame failed.  No Window"));
-        }
-    }
+		UGameEngine* GameEngine = Cast<UGameEngine>(GEngine);
+		if (MainWindow.IsValid() && GameEngine)
+		{
+			GameEngine->GameViewportWindow = MainWindow;
+		}
+		else
+		{
+			UE_LOG(LogPreLoadScreenManager, Warning, TEXT("FPreLoadScreenManager::PassLoadingScreenWindowBackToGame failed.  No Window"));
+		}
+	}
 }
 
 bool FPreLoadScreenManager::IsUsingMainWindow() const
 {
-    return MainWindow.IsValid();
+	return MainWindow.IsValid();
 }
 
 TSharedPtr<SWindow> FPreLoadScreenManager::GetRenderWindow()
 {
-    return MainWindow.IsValid() ? MainWindow.Pin() : nullptr;
+	return MainWindow.IsValid() ? MainWindow.Pin() : nullptr;
 }
 
 void FPreLoadScreenManager::WaitForEngineLoadingScreenToFinish()
@@ -743,27 +635,6 @@ bool FPreLoadScreenManager::ArePreLoadScreensEnabled()
 	return bEnabled;
 }
 
-void FPreLoadScreenManager::HandleFlushRenderingCommandsStart()
-{
-	//Whenever we flush rendering commands we need to unlock this critical section or we will softlock due to our enqueued rendering command never being able to
-	//acquire this critical section as its both locked on the GT and the GT won't unlock it as it's waiting for the rendering command to finish to unlock
-	AcquireCriticalSection.Unlock();
-
-	// Note: Don't modify bIsLocked as that will determine if we relock when flush completes
-}
-
-void FPreLoadScreenManager::HandleFlushRenderingCommandsEnd()
-{
-	//Relock critical section after the FlushRenderCommands finishes as we are
-	AcquireCriticalSection.Lock();
-
-	// Check if we were locked before, if not release lock to prevent flush from leaving us in undesired locked stated
-	if (!bIsLocked)
-	{
-		AcquireCriticalSection.Unlock();
-	}
-}
-
 void FPreLoadScreenManager::CleanUpResources()
 {
 	// Since we are on the game thread, the PreLoadScreen must be completed.
diff --git a/Engine/Source/Runtime/PreLoadScreen/Public/PreLoadScreenManager.h b/Engine/Source/Runtime/PreLoadScreen/Public/PreLoadScreenManager.h
--- a/Engine/Source/Runtime/PreLoadScreen/Public/PreLoadScreenManager.h
+++ b/Engine/Source/Runtime/PreLoadScreen/Public/PreLoadScreenManager.h
@@ -85,7 +85,7 @@ protected:
     //Separate tick that handles 
     PRELOADSCREEN_API void EarlyPlayFrameTick();
     PRELOADSCREEN_API void GameLogicFrameTick();
-    PRELOADSCREEN_API void EarlyPlayRenderFrameTick();
+	PRELOADSCREEN_API void EarlyPlayRenderFrameTick();
 	PRELOADSCREEN_API bool HasActivePreLoadScreenTypeForEarlyStartup() const;
 
 	PRELOADSCREEN_API void PlatformSpecificGameLogicFrameTick();
@@ -106,13 +106,6 @@ protected:
 	PRELOADSCREEN_API void HandleStopPreLoadScreen();
 	PRELOADSCREEN_API void CleanUpResources();
 
-	//Helpers that setup and clean-up delegates that only need to be active while we are playing an EarlyStartup PreLoadScreen
-	PRELOADSCREEN_API void RegisterDelegatesForEarlyStartupPlay();
-	PRELOADSCREEN_API void CleanUpDelegatesForEarlyStartupPlay();
-
-	PRELOADSCREEN_API void HandleFlushRenderingCommandsStart();
-	PRELOADSCREEN_API void HandleFlushRenderingCommandsEnd();
-
 	//Singleton Instance
 	struct FPreLoadScreenManagerDelete
 	{
@@ -124,10 +117,7 @@ protected:
 	static PRELOADSCREEN_API TUniquePtr<FPreLoadScreenManager, FPreLoadScreenManagerDelete> Instance;
 
 	static PRELOADSCREEN_API FCriticalSection AcquireCriticalSection;
-	static PRELOADSCREEN_API TAtomic<bool> bRenderingEnabled;
-	static PRELOADSCREEN_API TAtomic<bool> bIsLocked;
-	static PRELOADSCREEN_API TAtomic<bool> bIsLockedByGameThread;
-	static PRELOADSCREEN_API TAtomic<bool> bIsLockedByRenderThread;
+	static PRELOADSCREEN_API std::atomic<bool> bRenderingEnabled;
 
 	TArray<TSharedPtr<IPreLoadScreen>> PreLoadScreens;
 
@@ -135,7 +125,7 @@ protected:
     double LastTickTime;
 
     /** Widget renderer used to tick and paint windows in a thread safe way */
-    TSharedPtr<FPreLoadSlateWidgetRenderer, ESPMode::ThreadSafe> WidgetRenderer;
+    TSharedPtr<FPreLoadSlateWidgetRenderer> WidgetRenderer;
 
     /** The window that the loading screen resides in */
     TWeakPtr<class SWindow> MainWindow;
