Fix PreloadScreen deadlock on handle flush, when lock is already acquired.
#localization none

diff --git a/Engine/Source/Runtime/PreLoadScreen/Private/PreLoadScreenManager.cpp b/Engine/Source/Runtime/PreLoadScreen/Private/PreLoadScreenManager.cpp
--- a/Engine/Source/Runtime/PreLoadScreen/Private/PreLoadScreenManager.cpp
+++ b/Engine/Source/Runtime/PreLoadScreen/Private/PreLoadScreenManager.cpp
@@ -514,7 +514,7 @@ void FPreLoadScreenManager::EnableRendering(bool bEnabled)
 	FScopeLock ScopeLock(&AcquireCriticalSection);
 	TGuardValue LockedGuard(bIsLocked, true);
 	TGuardValue LockedGuardGameThread(bIsLockedByGameThread, IsInGameThread());
-	TGuardValue LockedGuardRenderThread(bIsLockedByGameThread, IsInRenderingThread());
+	TGuardValue LockedGuardRenderThread(bIsLockedByRenderThread, IsInRenderingThread());
 
     bRenderingEnabled = bEnabled;
 }
@@ -533,7 +533,7 @@ void FPreLoadScreenManager::EarlyPlayRenderFrameTick()
 			FScopeUnlock ScopeUnlock(&AcquireCriticalSection);
 			TGuardValue LockedGuard(bIsLocked, false);
 			TGuardValue LockedGuardGameThread(bIsLockedByGameThread, false);
-			TGuardValue LockedGuardRenderThread(bIsLockedByGameThread, false);
+			TGuardValue LockedGuardRenderThread(bIsLockedByRenderThread, false);
 			FPlatformProcess::Sleep(0);
 		}
 		return;
@@ -568,7 +568,7 @@ void FPreLoadScreenManager::EarlyPlayRenderFrameTick()
 					{
 						FScopeLock Lock(&FPreLoadScreenManager::AcquireCriticalSection);
 						TGuardValue LockedGuard(bIsLocked, true);
-						TGuardValue LockedGuardRenderThread(bIsLockedByGameThread, true);
+						TGuardValue LockedGuardRenderThread(bIsLockedByRenderThread, true);
 
 						// Self is still valid because we do a FlushRenderingCommands in StopPreLoadScreen
 						if (Self->bRenderingEnabled && Self->HasActivePreLoadScreenTypeForEarlyStartup() && !Self->bHasRenderPreLoadScreenFrame_RenderThread)
@@ -585,11 +585,18 @@ void FPreLoadScreenManager::EarlyPlayRenderFrameTick()
 
 			{
 				FScopeUnlock ScopeUnlock(&AcquireCriticalSection);
+
+				// Disable locking delegates during slate tick to prevent flushes inside slate tick from causing freezes.
+				CleanUpDelegatesForEarlyStartupPlay();
+
 				TGuardValue LockedGuard(bIsLocked, false);
 				TGuardValue LockedGuardGameThread(bIsLockedByGameThread, false);
-				TGuardValue LockedGuardRenderThread(bIsLockedByGameThread, false);
+				TGuardValue LockedGuardRenderThread(bIsLockedByRenderThread, false);
 				SlateApp.Tick();
 
+				// Re-enable delegates for flush handling
+				RegisterDelegatesForEarlyStartupPlay();
+
 				// Synchronize the game thread and the render thread so that the render thread doesn't get too far behind.
 				SlateApp.GetRenderer()->Sync();
 			}
@@ -625,7 +632,7 @@ void FPreLoadScreenManager::StopPreLoadScreen()
 			FScopeUnlock ScopeUnlock(&AcquireCriticalSection);
 			TGuardValue LockedGuard(bIsLocked, false);
 			TGuardValue LockedGuardGameThread(bIsLockedByGameThread, false);
-			TGuardValue LockedGuardRenderThread(bIsLockedByGameThread, false);
+			TGuardValue LockedGuardRenderThread(bIsLockedByRenderThread, false);
 			FlushRenderingCommands();
 		}
 	}
@@ -637,7 +644,7 @@ void FPreLoadScreenManager::HandleStopPreLoadScreen()
 		FScopeLock Lock(&AcquireCriticalSection); // prevent stop while we are rendering the preloadscreen
 		TGuardValue LockedGuard(bIsLocked, true);
 		TGuardValue LockedGuardGameThread(bIsLockedByGameThread, IsInGameThread());
-		TGuardValue LockedGuardRenderThread(bIsLockedByGameThread, IsInRenderingThread());
+		TGuardValue LockedGuardRenderThread(bIsLockedByRenderThread, IsInRenderingThread());
 
 		if (HasValidActivePreLoadScreen())
 		{
@@ -665,7 +672,7 @@ void FPreLoadScreenManager::PassPreLoadScreenWindowBackToGame() const
 		FScopeLock Lock(&AcquireCriticalSection); // wait until we finish with rendering before moving the context
 		TGuardValue LockedGuard(bIsLocked, true);
 		TGuardValue LockedGuardGameThread(bIsLockedByGameThread, IsInGameThread());
-		TGuardValue LockedGuardRenderThread(bIsLockedByGameThread, IsInRenderingThread());
+		TGuardValue LockedGuardRenderThread(bIsLockedByRenderThread, IsInRenderingThread());
 
         UGameEngine* GameEngine = Cast<UGameEngine>(GEngine);
         if (MainWindow.IsValid() && GameEngine)
