Clean up Preload Screen Manager locking and render commands
- Updated version that avoids taking AquireCriticalSection on the render thread

#localization none
#tests Tested on Windows, Lorry smoked on [redacted]

diff --git a/Engine/Source/Runtime/PreLoadScreen/Private/PreLoadScreenManager.cpp b/Engine/Source/Runtime/PreLoadScreen/Private/PreLoadScreenManager.cpp
--- a/Engine/Source/Runtime/PreLoadScreen/Private/PreLoadScreenManager.cpp
+++ b/Engine/Source/Runtime/PreLoadScreen/Private/PreLoadScreenManager.cpp
@@ -12,7 +12,6 @@
 
 #include "HAL/ThreadManager.h"
 #include "Modules/ModuleManager.h"
-#include "RenderingThread.h"
 
 #if BUILD_EMBEDDED_APP
 #include "Misc/EmbeddedCommunication.h"
@@ -29,11 +28,10 @@ IMPLEMENT_MODULE(FDefaultModuleImpl, PreLoadScreen);
 DEFINE_LOG_CATEGORY_STATIC(LogPreLoadScreenManager, Log, All);
 
 TUniquePtr<FPreLoadScreenManager, FPreLoadScreenManager::FPreLoadScreenManagerDelete> FPreLoadScreenManager::Instance;
-FCriticalSection FPreLoadScreenManager::AcquireCriticalSection;
-TAtomic<bool> FPreLoadScreenManager::bRenderingEnabled(true);
-TAtomic<bool> FPreLoadScreenManager::bIsLocked(false);
-TAtomic<bool> FPreLoadScreenManager::bIsLockedByGameThread(false);
-TAtomic<bool> FPreLoadScreenManager::bIsLockedByRenderThread(false);
+std::atomic<bool> FPreLoadScreenManager::bRenderingEnabled(true);
+
+FCriticalSection FPreLoadScreenManager::ActivePreloadScreenCriticalSection;
+TWeakPtr<IPreLoadScreen> FPreLoadScreenManager::ActivePreloadScreen;
 
 FPreLoadScreenManager* FPreLoadScreenManager::Get()
 {
@@ -44,11 +42,6 @@ void FPreLoadScreenManager::Create()
 {
 	check(IsInGameThread());
 
-	// Lock in case a user decide to create/detroy the manager multiple times
-	FScopeLock Lock(&AcquireCriticalSection);
-	TGuardValue LockedGuard(bIsLocked, true);
-	TGuardValue LockedGuardGameThread(bIsLockedByGameThread, true);
-
 	if (!Instance.IsValid() && ArePreLoadScreensEnabled())
 	{
 		Instance = TUniquePtr<FPreLoadScreenManager, FPreLoadScreenManager::FPreLoadScreenManagerDelete>(new FPreLoadScreenManager);
@@ -59,10 +52,6 @@ void FPreLoadScreenManager::Destroy()
 {
 	check(IsInGameThread());
 
-	FScopeLock Lock(&AcquireCriticalSection); // Make sure the render thread is completed before cleaning up
-	TGuardValue LockedGuard(bIsLocked, true);
-	TGuardValue LockedGuardGameThread(bIsLockedByGameThread, true);
-
 	if (Instance.IsValid())
 	{
 		Instance->CleanUpResources();
@@ -131,10 +120,6 @@ void FPreLoadScreenManager::UnRegisterPreLoadScreen(const TSharedPtr<IPreLoadScr
 {
 	check(IsInGameThread());
 
-	FScopeLock Lock(&AcquireCriticalSection);
-	TGuardValue LockedGuard(bIsLocked, true);
-	TGuardValue LockedGuardGameThread(bIsLockedByGameThread, true);
-
 	if (PreLoadScreen.IsValid())
 	{
 		const int32 IndexOf = PreLoadScreens.IndexOfByKey(PreLoadScreen);
@@ -179,13 +164,15 @@ void FPreLoadScreenManager::PlayPreLoadScreenAtIndex(int32 Index)
 
 	if (ArePreLoadScreensEnabled())
 	{
-		FScopeLock Lock(&AcquireCriticalSection);
-		TGuardValue LockedGuard(bIsLocked, true);
-		TGuardValue LockedGuardGameThread(bIsLockedByGameThread, true);
-
 		if (ensureAlwaysMsgf(!HasValidActivePreLoadScreen(), TEXT("Call to FPreLoadScreenManager::PlayPreLoadScreenAtIndex when something is already playing.")))
 		{
 			ActivePreLoadScreenIndex = Index;
+			if (HasValidActivePreLoadScreen())
+			{
+				FScopeLock PreloadScreenLock(&ActivePreloadScreenCriticalSection);
+				FPreLoadScreenManager::ActivePreloadScreen = PreLoadScreens[ActivePreLoadScreenIndex];
+			}
+
 			if (ensureAlwaysMsgf(HasValidActivePreLoadScreen(), TEXT("Call to FPreLoadScreenManager::PlayPreLoadScreenAtIndex with an invalid index! Nothing will play!")))
 			{
 				IPreLoadScreen* ActiveScreen = GetActivePreLoadScreen();
@@ -208,6 +195,11 @@ void FPreLoadScreenManager::PlayPreLoadScreenAtIndex(int32 Index)
 			}
 			else
 			{
+				{
+					FScopeLock PreloadScreenLock(&ActivePreloadScreenCriticalSection);
+					FPreLoadScreenManager::ActivePreloadScreen.Reset();
+				}
+
 				ActivePreLoadScreenIndex = INDEX_NONE;
 			}
 		}
@@ -254,8 +246,6 @@ void FPreLoadScreenManager::HandleEarlyStartupPlay()
 
 			FPlatformMisc::HidePlatformStartupScreen();
 
-			RegisterDelegatesForEarlyStartupPlay();			
-
 			{
 				SCOPED_BOOT_TIMING("FPreLoadScreenManager::EarlyPlayFrameTick()");
 
@@ -266,8 +256,6 @@ void FPreLoadScreenManager::HandleEarlyStartupPlay()
 				}
 			}
 
-			CleanUpDelegatesForEarlyStartupPlay();
-
 			if (bDidDisableScreensaver)
 			{
 				FPlatformApplicationMisc::ControlScreensaver(FGenericPlatformApplicationMisc::EScreenSaverAction::Enable);
@@ -278,19 +266,6 @@ void FPreLoadScreenManager::HandleEarlyStartupPlay()
 	}
 }
 
-void FPreLoadScreenManager::RegisterDelegatesForEarlyStartupPlay()
-{
-	//Have to register to handle FlushRenderingCommands for the length of the PreLoadScreen
-	FCoreRenderDelegates::OnFlushRenderingCommandsStart.AddRaw(this, &FPreLoadScreenManager::HandleFlushRenderingCommandsStart);
-	FCoreRenderDelegates::OnFlushRenderingCommandsEnd.AddRaw(this, &FPreLoadScreenManager::HandleFlushRenderingCommandsEnd);
-}
-
-void FPreLoadScreenManager::CleanUpDelegatesForEarlyStartupPlay()
-{
-	FCoreRenderDelegates::OnFlushRenderingCommandsStart.RemoveAll(this);
-	FCoreRenderDelegates::OnFlushRenderingCommandsEnd.RemoveAll(this);
-}
-
 void FPreLoadScreenManager::HandleEngineLoadingPlay()
 {
 	if (ensureAlwaysMsgf(HasActivePreLoadScreenType(EPreLoadScreenTypes::EngineLoadingScreen), TEXT("Invalid Active PreLoadScreen!")))
@@ -315,7 +290,7 @@ void FPreLoadScreenManager::HandleEngineLoadingPlay()
 
 		if (WidgetRenderer.IsValid())
 		{
-			if (ensure(SyncMechanism == nullptr))
+			if (SyncMechanism == nullptr)
 			{
 				SyncMechanism = new FPreLoadScreenSlateSynchMechanism(WidgetRenderer);
 				SyncMechanism->Initialize();
@@ -369,10 +344,6 @@ void FPreLoadScreenManager::StaticRenderTick_RenderThread()
 	LLM_SCOPE(ELLMTag::RenderingThreadMemory);
 	check(IsInRenderingThread());
 
-	FScopeLock Lock(&AcquireCriticalSection);
-	TGuardValue LockedGuard(bIsLocked, true);
-	TGuardValue LockedGuardGameThread(bIsLockedByRenderThread, true);
-
 	if (ensure(FPreLoadScreenManager::Get())) // The manager should clear the slate render thread before closing
 	{
 		FPreLoadScreenManager::Get()->RenderTick_RenderThread();
@@ -386,17 +357,17 @@ void FPreLoadScreenManager::RenderTick_RenderThread()
 	double DeltaTime = CurrentTime - LastRenderTickTime;
 	LastRenderTickTime = CurrentTime;
 
+	FScopeLock PreloadScreenLock(&ActivePreloadScreenCriticalSection);
+	TSharedPtr<IPreLoadScreen> PinnedActivePreloadScreen = FPreLoadScreenManager::ActivePreloadScreen.Pin();
+
 	//Check if we have an active index before doing any work
-	if (HasValidActivePreLoadScreen() && bRenderingEnabled)
+	if (PinnedActivePreloadScreen && bRenderingEnabled)
 	{
-		IPreLoadScreen* PreLoadScreen = GetActivePreLoadScreen();
-
-		check(PreLoadScreen && IsInRenderingThread());
-		if (MainWindow.IsValid() && VirtualRenderWindow.IsValid() && !PreLoadScreen->IsDone())
+		if (MainWindow.IsValid() && VirtualRenderWindow.IsValid() && !PinnedActivePreloadScreen->IsDone())
 		{
 			GFrameNumberRenderThread++;
 			GRHICommandList.GetImmediateCommandList().BeginFrame();
-			PreLoadScreen->RenderTick(DeltaTime);
+			PinnedActivePreloadScreen->RenderTick(DeltaTime);
 			GRHICommandList.GetImmediateCommandList().EndFrame();
 			GRHICommandList.GetImmediateCommandList().ImmediateFlush(EImmediateFlushType::FlushRHIThreadFlushResources);
 		}
@@ -511,37 +482,23 @@ void FPreLoadScreenManager::PlatformSpecificGameLogicFrameTick()
 
 void FPreLoadScreenManager::EnableRendering(bool bEnabled)
 {
-	FScopeLock ScopeLock(&AcquireCriticalSection);
-	TGuardValue LockedGuard(bIsLocked, true);
-	TGuardValue LockedGuardGameThread(bIsLockedByGameThread, IsInGameThread());
-	TGuardValue LockedGuardRenderThread(bIsLockedByRenderThread, IsInRenderingThread());
-
     bRenderingEnabled = bEnabled;
 }
 
 void FPreLoadScreenManager::EarlyPlayRenderFrameTick()
 {
-	bool bIsResponsibleForRendering_Local = true;
-
 	if (!bRenderingEnabled || !FSlateApplication::IsInitialized())
 	{
 		// If rendering disabled, FPreLoadScreenManager is responsible for rendering but choosing not to, probably because the
 		// app is not in the foreground.
-
-		// Cycle lock to give a chance to another thread to re-enable rendering.
-		{
-			FScopeUnlock ScopeUnlock(&AcquireCriticalSection);
-			TGuardValue LockedGuard(bIsLocked, false);
-			TGuardValue LockedGuardGameThread(bIsLockedByGameThread, false);
-			TGuardValue LockedGuardRenderThread(bIsLockedByRenderThread, false);
-			FPlatformProcess::Sleep(0);
-		}
+		FPlatformProcess::Sleep(0);
 		return;
 	}
 
 	IPreLoadScreen* ActivePreLoadScreen = PreLoadScreens[ActivePreLoadScreenIndex].Get();
 	if (ensureAlwaysMsgf(ActivePreLoadScreen, TEXT("Invalid Active PreLoadScreen during EarlyPlayRenderFrameTick!")))
 	{
+		bool bIsResponsibleForRendering_Local = true;
 		if (!ActivePreLoadScreen->ShouldRender())
 		{
 			bIsResponsibleForRendering_Local = false;
@@ -557,62 +514,38 @@ void FPreLoadScreenManager::EarlyPlayRenderFrameTick()
 		{
 			FSlateApplication& SlateApp = FSlateApplication::Get();
 			float SlateDeltaTime = SlateApp.GetDeltaTime();
-			FPreLoadScreenManager* Self = this;
 
-			{
-				UE::RenderCommandPipe::FSyncScope SyncScope;
+			//Setup Slate Render Command
+			ENQUEUE_RENDER_COMMAND(BeginPreLoadScreenFrame)(
+				[this, SlateDeltaTime](FRHICommandListImmediate& RHICmdList)
+				{
+					FScopeLock PreloadScreenLock(&ActivePreloadScreenCriticalSection);
+					TSharedPtr<IPreLoadScreen> PinnedActivePreloadScreen = FPreLoadScreenManager::ActivePreloadScreen.Pin();
 
-				//Setup Slate Render Command
-				ENQUEUE_RENDER_COMMAND(BeginPreLoadScreenFrame)(
-					[Self, SlateDeltaTime](FRHICommandListImmediate& RHICmdList)
+					// this is still valid because we do a FlushRenderingCommands in StopPreLoadScreen
+					if (FPreLoadScreenManager::bRenderingEnabled && PinnedActivePreloadScreen && !bHasRenderPreLoadScreenFrame_RenderThread)
 					{
-						FScopeLock Lock(&FPreLoadScreenManager::AcquireCriticalSection);
-						TGuardValue LockedGuard(bIsLocked, true);
-						TGuardValue LockedGuardRenderThread(bIsLockedByRenderThread, true);
-
-						// Self is still valid because we do a FlushRenderingCommands in StopPreLoadScreen
-						if (Self->bRenderingEnabled && Self->HasActivePreLoadScreenTypeForEarlyStartup() && !Self->bHasRenderPreLoadScreenFrame_RenderThread)
-						{
-							GFrameNumberRenderThread++;
-							GRHICommandList.GetImmediateCommandList().BeginFrame();
-
-							Self->bHasRenderPreLoadScreenFrame_RenderThread = true;
-							IPreLoadScreen* ActivePreLoadScreen = Self->PreLoadScreens[Self->ActivePreLoadScreenIndex].Get();
-							ActivePreLoadScreen->RenderTick(SlateDeltaTime);
-						}
-					});
-			}
-
-			{
-				FScopeUnlock ScopeUnlock(&AcquireCriticalSection);
-
-				// Disable locking delegates during slate tick to prevent flushes inside slate tick from causing freezes.
-				CleanUpDelegatesForEarlyStartupPlay();
-
-				TGuardValue LockedGuard(bIsLocked, false);
-				TGuardValue LockedGuardGameThread(bIsLockedByGameThread, false);
-				TGuardValue LockedGuardRenderThread(bIsLockedByRenderThread, false);
-				SlateApp.Tick();
+						GFrameNumberRenderThread++;
+						GRHICommandList.GetImmediateCommandList().BeginFrame();
 
-				// Re-enable delegates for flush handling
-				RegisterDelegatesForEarlyStartupPlay();
+						bHasRenderPreLoadScreenFrame_RenderThread = true;
+						PinnedActivePreloadScreen->RenderTick(SlateDeltaTime);
+					}
+				});
 
-				// Synchronize the game thread and the render thread so that the render thread doesn't get too far behind.
-				SlateApp.GetRenderer()->Sync();
-			}
+			SlateApp.Tick();
 
-			{
-				UE::RenderCommandPipe::FSyncScope SyncScope;
+			// Synchronize the game thread and the render thread so that the render thread doesn't get too far behind.
+			SlateApp.GetRenderer()->Sync();
 
-				ENQUEUE_RENDER_COMMAND(FinishPreLoadScreenFrame)(
-					[Self](FRHICommandListImmediate& RHICmdList)
-					{
-						// Self is still valid because we do a FlushRenderingCommands in StopPreLoadScreen
-						Self->bHasRenderPreLoadScreenFrame_RenderThread = false;
-						GRHICommandList.GetImmediateCommandList().EndFrame();
-						GRHICommandList.GetImmediateCommandList().ImmediateFlush(EImmediateFlushType::FlushRHIThreadFlushResources);
-					});
-			}
+			ENQUEUE_RENDER_COMMAND(FinishPreLoadScreenFrame)(
+				[this](FRHICommandListImmediate& RHICmdList)
+				{
+					// this is still valid because we do a FlushRenderingCommands in StopPreLoadScreen
+					bHasRenderPreLoadScreenFrame_RenderThread = false;
+					GRHICommandList.GetImmediateCommandList().EndFrame();
+					GRHICommandList.GetImmediateCommandList().ImmediateFlush(EImmediateFlushType::FlushRHIThreadFlushResources);
+				});
 		}
 	}
 }
@@ -628,29 +561,23 @@ void FPreLoadScreenManager::StopPreLoadScreen()
 			HandleStopPreLoadScreen();
 		}
 
-		{
-			FScopeUnlock ScopeUnlock(&AcquireCriticalSection);
-			TGuardValue LockedGuard(bIsLocked, false);
-			TGuardValue LockedGuardGameThread(bIsLockedByGameThread, false);
-			TGuardValue LockedGuardRenderThread(bIsLockedByRenderThread, false);
-			FlushRenderingCommands();
-		}
+		FlushRenderingCommands();
 	}
 }
 
 void FPreLoadScreenManager::HandleStopPreLoadScreen()
 {
 	{
-		FScopeLock Lock(&AcquireCriticalSection); // prevent stop while we are rendering the preloadscreen
-		TGuardValue LockedGuard(bIsLocked, true);
-		TGuardValue LockedGuardGameThread(bIsLockedByGameThread, IsInGameThread());
-		TGuardValue LockedGuardRenderThread(bIsLockedByRenderThread, IsInRenderingThread());
-
 		if (HasValidActivePreLoadScreen())
 		{
 			PreLoadScreens[ActivePreLoadScreenIndex]->OnStop();
 		}
 
+		{
+			FScopeLock PreloadScreenLock(&ActivePreloadScreenCriticalSection);
+			FPreLoadScreenManager::ActivePreloadScreen.Reset();
+		}
+
 		ActivePreLoadScreenIndex = -1;
 
 		//Clear our window content
@@ -667,33 +594,28 @@ void FPreLoadScreenManager::HandleStopPreLoadScreen()
 
 void FPreLoadScreenManager::PassPreLoadScreenWindowBackToGame() const
 {
-    if (IsUsingMainWindow())
-    {
-		FScopeLock Lock(&AcquireCriticalSection); // wait until we finish with rendering before moving the context
-		TGuardValue LockedGuard(bIsLocked, true);
-		TGuardValue LockedGuardGameThread(bIsLockedByGameThread, IsInGameThread());
-		TGuardValue LockedGuardRenderThread(bIsLockedByRenderThread, IsInRenderingThread());
-
-        UGameEngine* GameEngine = Cast<UGameEngine>(GEngine);
-        if (MainWindow.IsValid() && GameEngine)
-        {
-            GameEngine->GameViewportWindow = MainWindow;
-        }
-        else
-        {
-            UE_LOG(LogPreLoadScreenManager, Warning, TEXT("FPreLoadScreenManager::PassLoadingScreenWindowBackToGame failed.  No Window"));
-        }
-    }
+	if (IsUsingMainWindow())
+	{
+		UGameEngine* GameEngine = Cast<UGameEngine>(GEngine);
+		if (MainWindow.IsValid() && GameEngine)
+		{
+			GameEngine->GameViewportWindow = MainWindow;
+		}
+		else
+		{
+			UE_LOG(LogPreLoadScreenManager, Warning, TEXT("FPreLoadScreenManager::PassLoadingScreenWindowBackToGame failed.  No Window"));
+		}
+	}
 }
 
 bool FPreLoadScreenManager::IsUsingMainWindow() const
 {
-    return MainWindow.IsValid();
+	return MainWindow.IsValid();
 }
 
 TSharedPtr<SWindow> FPreLoadScreenManager::GetRenderWindow()
 {
-    return MainWindow.IsValid() ? MainWindow.Pin() : nullptr;
+	return MainWindow.IsValid() ? MainWindow.Pin() : nullptr;
 }
 
 void FPreLoadScreenManager::WaitForEngineLoadingScreenToFinish()
@@ -750,27 +672,6 @@ bool FPreLoadScreenManager::ArePreLoadScreensEnabled()
 	return bEnabled;
 }
 
-void FPreLoadScreenManager::HandleFlushRenderingCommandsStart()
-{
-	//Whenever we flush rendering commands we need to unlock this critical section or we will softlock due to our enqueued rendering command never being able to
-	//acquire this critical section as its both locked on the GT and the GT won't unlock it as it's waiting for the rendering command to finish to unlock
-	AcquireCriticalSection.Unlock();
-
-	// Note: Don't modify bIsLocked as that will determine if we relock when flush completes
-}
-
-void FPreLoadScreenManager::HandleFlushRenderingCommandsEnd()
-{
-	//Relock critical section after the FlushRenderCommands finishes as we are
-	AcquireCriticalSection.Lock();
-
-	// Check if we were locked before, if not release lock to prevent flush from leaving us in undesired locked stated
-	if (!bIsLocked)
-	{
-		AcquireCriticalSection.Unlock();
-	}
-}
-
 void FPreLoadScreenManager::CleanUpResources()
 {
 	// Since we are on the game thread, the PreLoadScreen must be completed.
diff --git a/Engine/Source/Runtime/PreLoadScreen/Private/PreLoadSlateThreading.cpp b/Engine/Source/Runtime/PreLoadScreen/Private/PreLoadSlateThreading.cpp
--- a/Engine/Source/Runtime/PreLoadScreen/Private/PreLoadSlateThreading.cpp
+++ b/Engine/Source/Runtime/PreLoadScreen/Private/PreLoadSlateThreading.cpp
@@ -231,7 +231,6 @@ void FPreLoadScreenSlateSynchMechanism::RunMainLoop_SlateThread()
 			// as we don't want Slate to submit any more draw calls until we Resume.
 			if (GDynamicRHI && !GDynamicRHI->RHIIsRenderingSuspended())
 			{
-				FScopeLock Lock(&FPreLoadScreenManager::AcquireCriticalSection);
 				if (IsSlateMainLoopRunning_AnyThread() && FPreLoadScreenManager::bRenderingEnabled)
 				{
 					WidgetRenderer->DrawWindow(DeltaTime);
diff --git a/Engine/Source/Runtime/PreLoadScreen/Public/PreLoadScreenManager.h b/Engine/Source/Runtime/PreLoadScreen/Public/PreLoadScreenManager.h
--- a/Engine/Source/Runtime/PreLoadScreen/Public/PreLoadScreenManager.h
+++ b/Engine/Source/Runtime/PreLoadScreen/Public/PreLoadScreenManager.h
@@ -85,7 +85,7 @@ protected:
     //Separate tick that handles 
     PRELOADSCREEN_API void EarlyPlayFrameTick();
     PRELOADSCREEN_API void GameLogicFrameTick();
-    PRELOADSCREEN_API void EarlyPlayRenderFrameTick();
+	PRELOADSCREEN_API void EarlyPlayRenderFrameTick();
 	PRELOADSCREEN_API bool HasActivePreLoadScreenTypeForEarlyStartup() const;
 
 	PRELOADSCREEN_API void PlatformSpecificGameLogicFrameTick();
@@ -106,13 +106,6 @@ protected:
 	PRELOADSCREEN_API void HandleStopPreLoadScreen();
 	PRELOADSCREEN_API void CleanUpResources();
 
-	//Helpers that setup and clean-up delegates that only need to be active while we are playing an EarlyStartup PreLoadScreen
-	PRELOADSCREEN_API void RegisterDelegatesForEarlyStartupPlay();
-	PRELOADSCREEN_API void CleanUpDelegatesForEarlyStartupPlay();
-
-	PRELOADSCREEN_API void HandleFlushRenderingCommandsStart();
-	PRELOADSCREEN_API void HandleFlushRenderingCommandsEnd();
-
 	//Singleton Instance
 	struct FPreLoadScreenManagerDelete
 	{
@@ -122,12 +115,11 @@ protected:
 		}
 	};
 	static PRELOADSCREEN_API TUniquePtr<FPreLoadScreenManager, FPreLoadScreenManagerDelete> Instance;
+	
+	static PRELOADSCREEN_API std::atomic<bool> bRenderingEnabled;
 
-	static PRELOADSCREEN_API FCriticalSection AcquireCriticalSection;
-	static PRELOADSCREEN_API TAtomic<bool> bRenderingEnabled;
-	static PRELOADSCREEN_API TAtomic<bool> bIsLocked;
-	static PRELOADSCREEN_API TAtomic<bool> bIsLockedByGameThread;
-	static PRELOADSCREEN_API TAtomic<bool> bIsLockedByRenderThread;
+	static FCriticalSection ActivePreloadScreenCriticalSection;
+	static TWeakPtr<IPreLoadScreen> ActivePreloadScreen;
 
 	TArray<TSharedPtr<IPreLoadScreen>> PreLoadScreens;
 
@@ -135,7 +127,7 @@ protected:
     double LastTickTime;
 
     /** Widget renderer used to tick and paint windows in a thread safe way */
-    TSharedPtr<FPreLoadSlateWidgetRenderer, ESPMode::ThreadSafe> WidgetRenderer;
+    TSharedPtr<FPreLoadSlateWidgetRenderer> WidgetRenderer;
 
     /** The window that the loading screen resides in */
     TWeakPtr<class SWindow> MainWindow;
