Preload Screen Truncation warning fixes

diff --git a/Engine/Source/Runtime/PreLoadScreen/PreLoadScreen.Build.cs b/Engine/Source/Runtime/PreLoadScreen/PreLoadScreen.Build.cs
--- a/Engine/Source/Runtime/PreLoadScreen/PreLoadScreen.Build.cs
+++ b/Engine/Source/Runtime/PreLoadScreen/PreLoadScreen.Build.cs
@@ -39,5 +39,8 @@ public class PreLoadScreen : ModuleRules
 				}
 			);
 		}
+
+
+		UnsafeTypeCastWarningLevel = WarningLevel.Error;
 	}
 }
diff --git a/Engine/Source/Runtime/PreLoadScreen/Private/PreLoadScreenManager.cpp b/Engine/Source/Runtime/PreLoadScreen/Private/PreLoadScreenManager.cpp
--- a/Engine/Source/Runtime/PreLoadScreen/Private/PreLoadScreenManager.cpp
+++ b/Engine/Source/Runtime/PreLoadScreen/Private/PreLoadScreenManager.cpp
@@ -354,7 +354,7 @@ void FPreLoadScreenManager::RenderTick_RenderThread()
 {
 	//Calculate tick time
 	const double CurrentTime = FPlatformTime::Seconds();
-	double DeltaTime = CurrentTime - LastRenderTickTime;
+	float DeltaTime = static_cast<float>(CurrentTime - LastRenderTickTime);
 	LastRenderTickTime = CurrentTime;
 
 	FScopeLock PreloadScreenLock(&ActivePreloadScreenCriticalSection);
@@ -436,12 +436,12 @@ void FPreLoadScreenManager::GameLogicFrameTick()
         }
 
         double CurrentTime = FPlatformTime::Seconds();
-        double DeltaTime = CurrentTime - LastTickTime;
+        float DeltaTime = static_cast<float>(CurrentTime - LastTickTime);
         LastTickTime = CurrentTime;
 
 		//Clamp to what should be more then any max reasonable time. This is to help with cases of
 		//backgrounding or setting breakpoints to trigger huge ticks
-		const double MaxTickTime = 5.0;
+		const float MaxTickTime = 5.0;
 		DeltaTime = FMath::Min(DeltaTime, MaxTickTime);
 
         //We have to manually tick everything as we are looping the main thread here
diff --git a/Engine/Source/Runtime/PreLoadScreen/Private/PreLoadSlateThreading.cpp b/Engine/Source/Runtime/PreLoadScreen/Private/PreLoadSlateThreading.cpp
--- a/Engine/Source/Runtime/PreLoadScreen/Private/PreLoadSlateThreading.cpp
+++ b/Engine/Source/Runtime/PreLoadScreen/Private/PreLoadSlateThreading.cpp
@@ -233,7 +233,7 @@ void FPreLoadScreenSlateSynchMechanism::RunMainLoop_SlateThread()
 			{
 				if (IsSlateMainLoopRunning_AnyThread() && FPreLoadScreenManager::bRenderingEnabled)
 				{
-					WidgetRenderer->DrawWindow(DeltaTime);
+					WidgetRenderer->DrawWindow(static_cast<float>(DeltaTime));
 
 					bRenderCommandEnqeued = true;
 
