EngineLoading screen: support screen size changes before showing the engine loading screen

#rnx

diff --git a/Engine/Source/Runtime/PreLoadScreen/Private/PreLoadScreenManager.cpp b/Engine/Source/Runtime/PreLoadScreen/Private/PreLoadScreenManager.cpp
--- a/Engine/Source/Runtime/PreLoadScreen/Private/PreLoadScreenManager.cpp
+++ b/Engine/Source/Runtime/PreLoadScreen/Private/PreLoadScreenManager.cpp
@@ -275,9 +275,19 @@ void FPreLoadScreenManager::HandleEngineLoadingPlay()
 		{
 			PreLoadScreen->OnPlay(MainWindow.Pin());
 
+			// The screen size may have changed between the VirtualRenderWindow creation and now
+			TOptional<UE::Slate::FDeprecateVector2DResult> NewScreenSize;
+
 			if (PreLoadScreen->GetWidget().IsValid() && VirtualRenderWindow.IsValid())
 			{
 				VirtualRenderWindow->SetContent(PreLoadScreen->GetWidget().ToSharedRef());
+
+				const UE::Slate::FDeprecateVector2DResult CurrentScreenSize = MainWindow.Pin()->GetClientSizeInScreen();
+				if (VirtualRenderWindow->GetClientSizeInScreen() != CurrentScreenSize)
+				{
+					NewScreenSize.Emplace(CurrentScreenSize);
+					VirtualRenderWindow->SetCachedSize(CurrentScreenSize);
+				}
 			}
 
 			//Need to update bIsResponsibleForRendering as a PreLoadScreen may not have updated it before this point
@@ -285,6 +295,12 @@ void FPreLoadScreenManager::HandleEngineLoadingPlay()
 			{
 				bIsResponsibleForRendering = true;
 				IsResponsibleForRenderingDelegate.Broadcast(bIsResponsibleForRendering);
+
+				if (NewScreenSize.IsSet() && FSlateApplication::IsInitialized())
+				{
+					// Force the viewport to resize before rendering
+					FSlateApplication::Get().GetRenderer()->UpdateFullscreenState(MainWindow.Pin().ToSharedRef(), (uint32)NewScreenSize.GetValue().X, (uint32)NewScreenSize.GetValue().Y);
+				}
 			}
 		}
 
