Fix crash in in FPreLoadScreenManager::HandleEngineLoadingPlay.
Code added recently was using a weak ptr (MainWindow) without null check.

#tests Tested load screen of Lyra in engine with debugger to validate that code was properly executed.
#rnx

diff --git a/Engine/Source/Runtime/PreLoadScreen/Private/PreLoadScreenManager.cpp b/Engine/Source/Runtime/PreLoadScreen/Private/PreLoadScreenManager.cpp
--- a/Engine/Source/Runtime/PreLoadScreen/Private/PreLoadScreenManager.cpp
+++ b/Engine/Source/Runtime/PreLoadScreen/Private/PreLoadScreenManager.cpp
@@ -273,7 +273,8 @@ void FPreLoadScreenManager::HandleEngineLoadingPlay()
 		IPreLoadScreen* PreLoadScreen = GetActivePreLoadScreen();
 		if (PreLoadScreen)
 		{
-			PreLoadScreen->OnPlay(MainWindow.Pin());
+			TSharedPtr<SWindow> MainWindowPtr = MainWindow.Pin();
+			PreLoadScreen->OnPlay(MainWindowPtr);
 
 			// The screen size may have changed between the VirtualRenderWindow creation and now
 			TOptional<UE::Slate::FDeprecateVector2DResult> NewScreenSize;
@@ -282,11 +283,14 @@ void FPreLoadScreenManager::HandleEngineLoadingPlay()
 			{
 				VirtualRenderWindow->SetContent(PreLoadScreen->GetWidget().ToSharedRef());
 
-				const UE::Slate::FDeprecateVector2DResult CurrentScreenSize = MainWindow.Pin()->GetClientSizeInScreen();
-				if (VirtualRenderWindow->GetClientSizeInScreen() != CurrentScreenSize)
+				if (MainWindowPtr.IsValid())
 				{
-					NewScreenSize.Emplace(CurrentScreenSize);
-					VirtualRenderWindow->SetCachedSize(CurrentScreenSize);
+					const UE::Slate::FDeprecateVector2DResult CurrentScreenSize = MainWindowPtr->GetClientSizeInScreen();
+					if (VirtualRenderWindow->GetClientSizeInScreen() != CurrentScreenSize)
+					{
+						NewScreenSize.Emplace(CurrentScreenSize);
+						VirtualRenderWindow->SetCachedSize(CurrentScreenSize);
+					}
 				}
 			}
 
@@ -296,10 +300,10 @@ void FPreLoadScreenManager::HandleEngineLoadingPlay()
 				bIsResponsibleForRendering = true;
 				IsResponsibleForRenderingDelegate.Broadcast(bIsResponsibleForRendering);
 
-				if (NewScreenSize.IsSet() && FSlateApplication::IsInitialized())
+				if (NewScreenSize.IsSet() && FSlateApplication::IsInitialized() && MainWindowPtr.IsValid()) //MainWindowPtr should be valid if we get there, but... better safe than sorry.
 				{
 					// Force the viewport to resize before rendering
-					FSlateApplication::Get().GetRenderer()->UpdateFullscreenState(MainWindow.Pin().ToSharedRef(), (uint32)NewScreenSize.GetValue().X, (uint32)NewScreenSize.GetValue().Y);
+					FSlateApplication::Get().GetRenderer()->UpdateFullscreenState(MainWindowPtr.ToSharedRef(), (uint32)NewScreenSize.GetValue().X, (uint32)NewScreenSize.GetValue().Y);
 				}
 			}
 		}
