[Redesigned Boot Up Flow] iOS notifications permission and background download UI

#rnx

diff --git a/Engine/Source/Runtime/PreLoadScreen/Private/PreLoadScreenBase.cpp b/Engine/Source/Runtime/PreLoadScreen/Private/PreLoadScreenBase.cpp
--- a/Engine/Source/Runtime/PreLoadScreen/Private/PreLoadScreenBase.cpp
+++ b/Engine/Source/Runtime/PreLoadScreen/Private/PreLoadScreenBase.cpp
@@ -34,7 +34,17 @@ void FPreLoadScreenBase::InitSettingsFromConfig(const FString& ConfigFileName)
         float TimeToDisplayEachBackground = 5.0f;
         Config->GetFloat(*UIConfigSection, TEXT("TimeToDisplayEachBackground"), TimeToDisplayEachBackground);
         FPreLoadSettingsContainerBase::Get().TimeToDisplayEachBackground = TimeToDisplayEachBackground;
-
+    	
+        //Parse time to recheck notifications time
+        float TimeToRecheckNotificationsPermission = 1.0f;
+        Config->GetFloat(*UIConfigSection, TEXT("TimeToRecheckNotificationsPermission"), TimeToRecheckNotificationsPermission);
+    	FPreLoadSettingsContainerBase::Get().TimeToRecheckNotificationsPermission = TimeToRecheckNotificationsPermission;
+    	
+    	//Parse time to show the notifications permission modal
+    	float TimeToShowNotificationsPermissionModal = 2.0f;
+    	Config->GetFloat(*UIConfigSection, TEXT("TimeToShowNotificationsPermissionModal"), TimeToShowNotificationsPermissionModal);
+    	FPreLoadSettingsContainerBase::Get().TimeToShowNotificationsPermissionModal = TimeToShowNotificationsPermissionModal;
+        
 		//Parse LoadingGroups. You want to do this before ScreenGroupings and CustomImageBrushes
 		TArray<FString> LoadingGroups;
 		Config->GetArray(*UIConfigSection, TEXT("LoadingGroups"), LoadingGroups);
diff --git a/Engine/Source/Runtime/PreLoadScreen/Public/PreLoadSettingsContainer.h b/Engine/Source/Runtime/PreLoadScreen/Public/PreLoadSettingsContainer.h
--- a/Engine/Source/Runtime/PreLoadScreen/Public/PreLoadSettingsContainer.h
+++ b/Engine/Source/Runtime/PreLoadScreen/Public/PreLoadSettingsContainer.h
@@ -150,6 +150,8 @@ public:
 	PRELOADSCREEN_API virtual void SetShouldLoadBrushes(bool bInShouldLoadBrushes);
 
     float TimeToDisplayEachBackground;
+	float TimeToRecheckNotificationsPermission;
+	float TimeToShowNotificationsPermissionModal;
     
 	FName GetCurrentLoadGrouping() const { return CurrentLoadGroup; }
 	PRELOADSCREEN_API void LoadGrouping(FName Identifier);
