Gather text from source to include nested macros; fix unlocalized EarlyStartupPatcherPlugin text

diff --git a/Engine/Source/Runtime/PreLoadScreen/Private/PreLoadSettingsContainer.cpp b/Engine/Source/Runtime/PreLoadScreen/Private/PreLoadSettingsContainer.cpp
--- a/Engine/Source/Runtime/PreLoadScreen/Private/PreLoadSettingsContainer.cpp
+++ b/Engine/Source/Runtime/PreLoadScreen/Private/PreLoadSettingsContainer.cpp
@@ -390,7 +390,7 @@ void FPreLoadSettingsContainerBase::ParseLocalizedTextConfigString(const FString
 {
     TArray<FString> LocalizedTextComponents;
     ConfigEntry.ParseIntoArray(LocalizedTextComponents, TEXT(","), true);
-    if (ensureAlwaysMsgf(IsValidLocalizedTextConfigString(LocalizedTextComponents), TEXT("Invalid Localized Text Entry in config: Expected Format: +LocalizedText=(TextIdentifier, NS Localized Text) Config Entry: %s"), *ConfigEntry))
+    if (ensureAlwaysMsgf(IsValidLocalizedTextConfigString(LocalizedTextComponents), TEXT("Invalid Localized Text Entry in config: Expected Format: +LocalizedText=(TextIdentifier, NSLOCTEXT(Namespace, Key, Localized Text) or +LocalizedText=(TextIdentifier, Namespace, Key, Text) Config Entry: %s"), *ConfigEntry))
     {
         //Clean up the identifier to remove extra spaces and the first (
         FString Identifier = LocalizedTextComponents[0];
@@ -400,7 +400,8 @@ void FPreLoadSettingsContainerBase::ParseLocalizedTextConfigString(const FString
         //LocalizedTextComponents[1] is the NameSpace for the loctext
         FString LocNameSpace = LocalizedTextComponents[1];
         LocNameSpace.TrimStartAndEndInline();
-        LocNameSpace.RemoveFromStart("NSLOCTEXT(\"");
+        LocNameSpace.RemoveFromStart("NSLOCTEXT(");
+        LocNameSpace.RemoveFromStart("\"");
         LocNameSpace.RemoveFromEnd("\"");
 
         //LocalizedTextComponents[2] is the identifier for the FText
@@ -413,7 +414,7 @@ void FPreLoadSettingsContainerBase::ParseLocalizedTextConfigString(const FString
         FString LocInitialValue = LocalizedTextComponents[3];
         LocInitialValue.TrimStartAndEndInline();
         LocInitialValue.RemoveFromStart("\"");
-        LocInitialValue.RemoveFromEnd(")"); //remove these separately so that if the file is missing 1 ) or the " is out of order it still works
+        LocInitialValue.RemoveFromEnd(")"); //remove these separately so that if the file is missing one ) or the " is out of order it still works
         LocInitialValue.RemoveFromEnd(")");
         LocInitialValue.RemoveFromEnd("\"");
 
