ESP: time to update download size text added

#rnx

diff --git a/Engine/Source/Runtime/PreLoadScreen/Private/PreLoadScreenBase.cpp b/Engine/Source/Runtime/PreLoadScreen/Private/PreLoadScreenBase.cpp
--- a/Engine/Source/Runtime/PreLoadScreen/Private/PreLoadScreenBase.cpp
+++ b/Engine/Source/Runtime/PreLoadScreen/Private/PreLoadScreenBase.cpp
@@ -33,7 +33,12 @@ void FPreLoadScreenBase::InitSettingsFromConfig(const FString& ConfigFileName)
         //Parse background display time
         float TimeToDisplayEachBackground = 5.0f;
         Config->GetFloat(*UIConfigSection, TEXT("TimeToDisplayEachBackground"), TimeToDisplayEachBackground);
-        FPreLoadSettingsContainerBase::Get().TimeToDisplayEachBackground = TimeToDisplayEachBackground;
+    	FPreLoadSettingsContainerBase::Get().TimeToDisplayEachBackground = TimeToDisplayEachBackground;
+    	
+    	//Parse time to update the downloaded data size
+    	float TimeToUpdateDownloadedDataSize = 0.0f;
+    	Config->GetFloat(*UIConfigSection, TEXT("TimeToUpdateDownloadedDataSize"), TimeToUpdateDownloadedDataSize);
+    	FPreLoadSettingsContainerBase::Get().TimeToUpdateDownloadedDataSize = TimeToUpdateDownloadedDataSize;
     	
         //Parse time to recheck notifications time
         float TimeToRecheckNotificationsPermission = 1.0f;
diff --git a/Engine/Source/Runtime/PreLoadScreen/Public/PreLoadSettingsContainer.h b/Engine/Source/Runtime/PreLoadScreen/Public/PreLoadSettingsContainer.h
--- a/Engine/Source/Runtime/PreLoadScreen/Public/PreLoadSettingsContainer.h
+++ b/Engine/Source/Runtime/PreLoadScreen/Public/PreLoadSettingsContainer.h
@@ -150,6 +150,7 @@ public:
 	PRELOADSCREEN_API virtual void SetShouldLoadBrushes(bool bInShouldLoadBrushes);
 
     float TimeToDisplayEachBackground;
+	float TimeToUpdateDownloadedDataSize;
 	float TimeToRecheckNotificationsPermission;
 	float TimeToShowNotificationsPermissionModal;
     
