Fix crash when exiting with ALT+F4 during preinit phase of the engine.

It looks like appexit can be called many times when requesting exit that early, so I had to fix FSlateRHIResourceManager calling a ReleaseResources on a null pointer after a call to DeleteResources.
It will be useful however to investigate why the exit delegates are called 2 times in this situation.

#tests Closed the game at various stage (right after boot, during loading, in lobby). It always went well, while it was crashing 100% if doing it right after boot without the fix.
[RN] Fixed crash when closing the engine during the preinit phase.

diff --git a/Engine/Source/Runtime/PreLoadScreen/Private/PreLoadScreenManager.cpp b/Engine/Source/Runtime/PreLoadScreen/Private/PreLoadScreenManager.cpp
--- a/Engine/Source/Runtime/PreLoadScreen/Private/PreLoadScreenManager.cpp
+++ b/Engine/Source/Runtime/PreLoadScreen/Private/PreLoadScreenManager.cpp
@@ -71,6 +71,14 @@ FPreLoadScreenManager::FPreLoadScreenManager()
 	, bIsEngineLoadingComplete(false)
 {}
 
+void FPreLoadScreenManager::OnRequestExit()
+{
+	if (Instance.IsValid() && ArePreLoadScreensEnabled())
+	{
+		RequestEngineExit(TEXT("Window closed during preload"));
+	}
+}
+
 void FPreLoadScreenManager::Initialize(FSlateRenderer& InSlateRenderer)
 {
 	check(IsInGameThread());
@@ -107,6 +115,9 @@ void FPreLoadScreenManager::Initialize(FSlateRenderer& InSlateRenderer)
 
 	LastRenderTickTime = FPlatformTime::Seconds();
 	LastTickTime = FPlatformTime::Seconds();
+
+	//This delegate will stay registered until another system do it. However, it will do something only for the lifetime of FPreLoadScreenManager.
+	FSlateApplication::Get().SetExitRequestedHandler(FSimpleDelegate::CreateStatic(&FPreLoadScreenManager::OnRequestExit));
 }
 
 void FPreLoadScreenManager::RegisterPreLoadScreen(const TSharedPtr<IPreLoadScreen>& PreLoadScreen)
@@ -162,6 +173,13 @@ void FPreLoadScreenManager::PlayPreLoadScreenAtIndex(int32 Index)
 {
 	check(IsInGameThread());
 
+	if (IsEngineExitRequested())
+	{	//Don't try to display a new preload screen when the engine is closing:
+		//	-as we are closing we don't want to display load feedback, which is the opposite of what we are trying to do (but we cannot actually close in the middle of the init sequence).
+		//	-we may no longer have a SWindow to use for display.
+		return;
+	}
+
 	if (ArePreLoadScreensEnabled())
 	{
 		if (ensureAlwaysMsgf(!HasValidActivePreLoadScreen(), TEXT("Call to FPreLoadScreenManager::PlayPreLoadScreenAtIndex when something is already playing.")))
@@ -203,7 +221,7 @@ void FPreLoadScreenManager::PlayPreLoadScreenAtIndex(int32 Index)
 				ActivePreLoadScreenIndex = INDEX_NONE;
 			}
 		}
-    }
+	}
 }
 
 bool FPreLoadScreenManager::PlayPreLoadScreenWithTag(FName InTag)
@@ -316,8 +334,8 @@ void FPreLoadScreenManager::HandleEngineLoadingPlay()
 				SyncMechanism->Initialize();
 			}
 		}
+		}
 	}
-}
 
 void FPreLoadScreenManager::HandleCustomSplashScreenPlay()
 {
diff --git a/Engine/Source/Runtime/PreLoadScreen/Public/PreLoadScreenManager.h b/Engine/Source/Runtime/PreLoadScreen/Public/PreLoadScreenManager.h
--- a/Engine/Source/Runtime/PreLoadScreen/Public/PreLoadScreenManager.h
+++ b/Engine/Source/Runtime/PreLoadScreen/Public/PreLoadScreenManager.h
@@ -106,6 +106,8 @@ protected:
 	PRELOADSCREEN_API void HandleStopPreLoadScreen();
 	PRELOADSCREEN_API void CleanUpResources();
 
+	static void OnRequestExit();
+
 	//Singleton Instance
 	struct FPreLoadScreenManagerDelete
 	{
diff --git a/Engine/Source/Runtime/SlateRHIRenderer/Private/SlateRHIResourceManager.cpp b/Engine/Source/Runtime/SlateRHIRenderer/Private/SlateRHIResourceManager.cpp
--- a/Engine/Source/Runtime/SlateRHIRenderer/Private/SlateRHIResourceManager.cpp
+++ b/Engine/Source/Runtime/SlateRHIRenderer/Private/SlateRHIResourceManager.cpp
@@ -1069,7 +1069,10 @@ void FSlateRHIResourceManager::ReleaseResources()
 
 	DynamicResourceMap.ReleaseResources();
 
-	VectorGraphicsCache->ReleaseResources();
+	if (VectorGraphicsCache != nullptr)
+	{
+		VectorGraphicsCache->ReleaseResources();
+	}
 
 	// Note the base class has texture proxies only which do not need to be released
 }
