Added feature to allow preload images to load async

diff --git a/Engine/Source/Runtime/Launch/Private/LaunchEngineLoop.cpp b/Engine/Source/Runtime/Launch/Private/LaunchEngineLoop.cpp
--- a/Engine/Source/Runtime/Launch/Private/LaunchEngineLoop.cpp
+++ b/Engine/Source/Runtime/Launch/Private/LaunchEngineLoop.cpp
@@ -58,6 +58,7 @@
 #include "SceneInterface.h"
 #include "StereoRenderUtils.h"
 #include "DataDrivenShaderPlatformInfo.h"
+#include "PreLoadSettingsContainer.h"
 #endif
 #if WITH_APPLICATION_CORE
 #include "HAL/PlatformApplicationMisc.h"
@@ -2934,6 +2935,13 @@ int32 FEngineLoop::PreInitPreStartupScreen(const TCHAR* CmdLine)
 		FInternationalization::Get().LoadAllCultureData();
 	}
 
+#if WITH_ENGINE
+	{
+		SCOPED_BOOT_TIMING("Begin Async Pre-Load Images");
+		FPreLoadSettingsContainerBase::Get().BeginAsyncLoadCustomBrushImages();
+	}
+#endif
+
 	FEmbeddedCommunication::ForceTick(3);
 
 	SlowTask.EnterProgressFrame(10);
diff --git a/Engine/Source/Runtime/PreLoadScreen/PreLoadScreen.Build.cs b/Engine/Source/Runtime/PreLoadScreen/PreLoadScreen.Build.cs
--- a/Engine/Source/Runtime/PreLoadScreen/PreLoadScreen.Build.cs
+++ b/Engine/Source/Runtime/PreLoadScreen/PreLoadScreen.Build.cs
@@ -21,6 +21,7 @@ public class PreLoadScreen : ModuleRules
 					"InputCore",
 					"RenderCore",
 					"CoreUObject",
+					"ImageCore",
 					"RHI",
 					"Slate",
 					"SlateCore",
diff --git a/Engine/Source/Runtime/PreLoadScreen/Private/PreLoadScreenBase.cpp b/Engine/Source/Runtime/PreLoadScreen/Private/PreLoadScreenBase.cpp
--- a/Engine/Source/Runtime/PreLoadScreen/Private/PreLoadScreenBase.cpp
+++ b/Engine/Source/Runtime/PreLoadScreen/Private/PreLoadScreenBase.cpp
@@ -2,7 +2,6 @@
 
 #include "PreLoadScreenBase.h"
 
-#include "Interfaces/IPluginManager.h"
 #include "PreLoadSettingsContainer.h"
 
 #include "Misc/ConfigCacheIni.h"
@@ -14,16 +13,11 @@ void FPreLoadScreenBase::InitSettingsFromConfig(const FString& ConfigFileName)
     const FString UIConfigSection = TEXT("PreLoadScreen.UISettings");
 
     //Find plugin content path from name by going through enabled content plugins and see if one matches.
-    FString PluginContentDir;
-    for (TSharedRef<IPlugin> Plugin : IPluginManager::Get().GetEnabledPlugins())
+    TOptional<FString> PluginContentDir = FPreLoadSettingsContainerBase::FindPluginContentDirIgnoreCase(PluginName);
+    if (PluginContentDir.IsSet())
     {
-        if (Plugin->CanContainContent() && (Plugin->GetName().Equals(PluginName, ESearchCase::IgnoreCase)))
-        {
-            PluginContentDir = Plugin->GetContentDir();
-            break;
-        }
+        FPreLoadSettingsContainerBase::Get().SetPluginContentDir(*PluginContentDir);
     }
-    FPreLoadSettingsContainerBase::Get().SetPluginContentDir(PluginContentDir);
 
     FConfigFile* Config = GConfig->FindConfigFileWithBaseName(*ConfigFileName);
     if (ensureAlwaysMsgf(Config, TEXT("Unable to find .ini file for %s"), *ConfigFileName))
diff --git a/Engine/Source/Runtime/PreLoadScreen/Private/PreLoadSettingsContainer.cpp b/Engine/Source/Runtime/PreLoadScreen/Private/PreLoadSettingsContainer.cpp
--- a/Engine/Source/Runtime/PreLoadScreen/Private/PreLoadSettingsContainer.cpp
+++ b/Engine/Source/Runtime/PreLoadScreen/Private/PreLoadSettingsContainer.cpp
@@ -7,14 +7,36 @@
 #include "Fonts/UnicodeBlockRange.h"
 #include "Misc/FileHelper.h"
 #include "Misc/Paths.h"
+#include "Misc/ConfigCacheIni.h"
+#include "Modules/ModuleManager.h"
 #include "Framework/Application/SlateApplication.h"
+#include "ImageUtils.h"
+#include "IImageWrapperModule.h"
+#include "Interfaces/IPluginManager.h"
 
 #include "Internationalization/Culture.h"
 
+DEFINE_LOG_CATEGORY_STATIC(LogPreLoadSettingsContainer, Log, All);
+
 FString FPreLoadSettingsContainerBase::UseSystemFontOverride = TEXT("SYSTEM");
 FString FPreLoadSettingsContainerBase::DefaultInitialLoadingGroupIdentifier = TEXT("InitialLoad");
 FPreLoadSettingsContainerBase* FPreLoadSettingsContainerBase::Instance = nullptr;
 
+TOptional<FString> FPreLoadSettingsContainerBase::FindPluginContentDirIgnoreCase(const FString& PluginName)
+{
+    //Find plugin content path from name by going through enabled content plugins and see if one matches.
+    FString PluginContentDir;
+    for (TSharedRef<IPlugin> Plugin : IPluginManager::Get().GetEnabledPlugins())
+    {
+        if (Plugin->CanContainContent() && (Plugin->GetName().Equals(PluginName, ESearchCase::IgnoreCase)))
+        {
+            return Plugin->GetContentDir();
+        }
+    }
+
+    return {};
+}
+
 FPreLoadSettingsContainerBase::~FPreLoadSettingsContainerBase()
 {
 	for (auto& KVPair : BrushResources)
@@ -214,6 +236,26 @@ void FPreLoadSettingsContainerBase::ParseScreenOrderConfigString(const FString&
 
 void FPreLoadSettingsContainerBase::PerformInitialAssetLoad()
 {
+	if (PreloadImagesTask.IsValid())
+	{
+		PreloadImagesTask.Wait();
+		TArray<FPreLoadImage> PreloadImages = MoveTemp(PreloadImagesTask.GetResult());
+
+		for (FPreLoadImage& PreLoadImage : PreloadImages)
+		{
+			FSlateTextureDataRef TextureStorage = MakeShareable(
+				new FSlateTextureData(
+					PreLoadImage.SizeX,
+					PreLoadImage.SizeY,
+					GPixelFormats[PF_B8G8R8A8].BlockBytes,
+					MoveTemp(PreLoadImage.RawData)
+				)
+			);
+
+			FSlateApplication::Get().GetRenderer()->GenerateDynamicImageResource(PreLoadImage.FilePath, TextureStorage);
+		}
+	}
+
 	//Try and load the default initial loading group, but only if we know we have it to avoid triggering any ensures
 	FCustomBrushLoadingGroup* FoundDefaultLoadingGroup = BrushLoadingGroups.Find(*DefaultInitialLoadingGroupIdentifier);
 	if (FoundDefaultLoadingGroup != nullptr)
@@ -222,6 +264,90 @@ void FPreLoadSettingsContainerBase::PerformInitialAssetLoad()
 	}
 }
 
+void FPreLoadSettingsContainerBase::BeginAsyncLoadCustomBrushImages()
+{	
+	TArray<FString> AsyncPreLoadScreenPlugins;
+	GConfig->GetArray(TEXT("PreLoadScreen.AsyncPreload"), TEXT("AsyncPreLoadScreenPlugins"), AsyncPreLoadScreenPlugins, GGameIni);
+
+	const FString UIConfigSection = TEXT("PreLoadScreen.UISettings");
+
+    // this is the same logic that happens in FPreLoadScreenBase::InitSettingsFromConfig
+    // to pre-parse the .ini files and start a background load job, so when the images are requested later they are already present.
+    //
+    // The pre-load screens themselves have a chicken and egg scenario where they need the window to create the slate renderer
+    // and they need the slate renderer to load the images.
+    // 
+    // So doing things here lets us cache the images outside of that
+
+	for (const FString& PluginName : AsyncPreLoadScreenPlugins)
+	{
+    	FConfigFile* Config = GConfig->FindConfigFileWithBaseName(*PluginName);	
+    	if (!Config)
+    	{
+    		UE_LOG(LogPreLoadSettingsContainer, Error, TEXT("Config file '%s' not found"), *PluginName);
+    		continue;
+    	}
+
+    	TOptional<FString> OptPluginContentDir = FindPluginContentDirIgnoreCase(PluginName);
+    	if (!OptPluginContentDir.IsSet())
+    	{
+	    	UE_LOG(LogPreLoadSettingsContainer, Error, TEXT("Plugin '%s' not found"), *PluginName)
+	    	continue;
+	    }
+
+	    // this must be set so the Parse* functions below function correctly
+	    SetPluginContentDir(*OptPluginContentDir);
+
+		TArray<FString> LoadingGroups;
+		Config->GetArray(*UIConfigSection, TEXT("LoadingGroups"), LoadingGroups);
+		ParseLoadingGroups(LoadingGroups);
+
+	    TArray<FString> CustomImageBrushes;
+        Config->GetArray(*UIConfigSection, TEXT("CustomImageBrushes"), CustomImageBrushes);
+        for (const FString& BrushConfigEntry : CustomImageBrushes)
+        {
+            ParseBrushConfigEntry(BrushConfigEntry);
+        } 
+	}
+
+	// ensure the module we need is present
+	FModuleManager::LoadModuleChecked<IImageWrapperModule>("ImageWrapper");
+
+	FCustomBrushLoadingGroup* FoundDefaultLoadingGroup = BrushLoadingGroups.Find(*DefaultInitialLoadingGroupIdentifier);
+	if (FoundDefaultLoadingGroup != nullptr)
+	{
+		// start this group loading as one job
+		PreloadImagesTask = UE::Tasks::Launch(TEXT("Async Preload Images"), [BrushesToLoad = FoundDefaultLoadingGroup->CustomBrushDefinesToLoad]() -> TArray<FPreLoadImage> {
+
+			SCOPED_BOOT_TIMING("FPreLoadSettingsContainerBase::LoadGrouping Async");
+
+			TArray<FPreLoadImage> Result;
+
+			for (const FCustomBrushDefine& BrushDefine : BrushesToLoad)
+			{
+				FImage Image;
+
+				if (!FImageUtils::LoadImage(*BrushDefine.FilePath, Image))
+				{
+					UE_LOG(LogPreLoadSettingsContainer, Log, TEXT("Could not load file for Slate resource: [%s] file: [%s]"), *BrushDefine.BrushIdentifier, *BrushDefine.FilePath );
+					Result.Pop();
+					continue;
+				}
+				Image.ChangeFormat(ERawImageFormat::BGRA8, EGammaSpace::sRGB);
+
+				FPreLoadImage& Texture = Result.Emplace_GetRef();
+
+				Texture.FilePath = *BrushDefine.FilePath;
+				Texture.RawData = TArray<uint8>(MoveTemp(Image.RawData));
+				Texture.SizeX = Image.SizeX;
+				Texture.SizeY = Image.SizeY;
+			}
+
+			return Result;
+		});
+	}
+}
+
 void FPreLoadSettingsContainerBase::LoadGrouping(FName Identifier)
 {
 	FCustomBrushLoadingGroup* FoundBrushLoadingGroup = BrushLoadingGroups.Find(*Identifier.ToString());
@@ -526,4 +652,4 @@ const FPreLoadSettingsContainerBase::FScreenGroupingBase* FPreLoadSettingsContai
 bool FPreLoadSettingsContainerBase::IsValidScreenIndex(int index) const
 {
 	return (!CurrentLoadGroup.IsNone() && CurrentLoadGroup.IsValid() && (ScreenOrderByLoadingGroups.Find(CurrentLoadGroup) != nullptr)) ? ScreenOrderByLoadingGroups.Find(CurrentLoadGroup)->ScreenGroupings.IsValidIndex(index) : false;
-}
\ No newline at end of file
+}
diff --git a/Engine/Source/Runtime/PreLoadScreen/Public/PreLoadSettingsContainer.h b/Engine/Source/Runtime/PreLoadScreen/Public/PreLoadSettingsContainer.h
--- a/Engine/Source/Runtime/PreLoadScreen/Public/PreLoadSettingsContainer.h
+++ b/Engine/Source/Runtime/PreLoadScreen/Public/PreLoadSettingsContainer.h
@@ -5,11 +5,13 @@
 #include "CoreMinimal.h"
 #include "RenderDeferredCleanup.h"
 #include "UObject/GCObject.h"
+#include "Tasks/Task.h"
 
 struct FSlateDynamicImageBrush;
 struct FCompositeFont;
 struct FStandaloneCompositeFont;
 
+
 //This is a helper class that we use to hold values we parse from the .ini. Clean way to access things like dynamic image brushes / fonts / etc used in our UI that
 //we want to be somewhat data driven but we can't rely on UObject support to implement(as the PreLoad stuff happens too early for UObject support)
 //This lets us set easy to change values in our .ini that are parsed at runtime and stored in this container
@@ -69,6 +71,16 @@ public:
 		TArray<FCustomBrushDefine> CustomBrushDefinesToLoad;
 	};
 
+	//Pre-load image bytes
+	struct FPreLoadImage
+	{
+		int32 SizeX = 0;
+		int32 SizeY = 0;
+		TArray<uint8> RawData;
+
+		FName FilePath;
+	};
+
 public:
 
     static FPreLoadSettingsContainerBase& Get()
@@ -153,12 +165,16 @@ public:
     
 	FName GetCurrentLoadGrouping() const { return CurrentLoadGroup; }
 	PRELOADSCREEN_API void LoadGrouping(FName Identifier);
+	PRELOADSCREEN_API void BeginAsyncLoadCustomBrushImages();
 	PRELOADSCREEN_API void PerformInitialAssetLoad();
 
     //Helper function that takes in a file path and tries to reconsile it to be Plugin Specific if applicable.
     //Ensures if file is not found in either Plugin's content dir or the original path
     PRELOADSCREEN_API virtual FString ConvertIfPluginRelativeContentPath(const FString& FilePath);
 
+    // Helper function to search plugin names case insensitively to find the content directory
+	static TOptional<FString> FindPluginContentDirIgnoreCase(const FString& PluginName);
+
 protected:
 
     //Helper functions that verify if the supplied .ini config entry is valid to create a resource out of it
@@ -180,6 +196,8 @@ protected:
 
 	TMap<FName, FCustomBrushLoadingGroup> BrushLoadingGroups;
 
+	UE::Tasks::TTask<TArray<FPreLoadImage>> PreloadImagesTask;
+
     //This string is used to make file paths relative to a particular Plugin's content directory when parsing file paths.
     FString PluginContentDir;
 
