[Linux]Re-enable PreLoadScreens

diff --git a/Engine/Source/Runtime/PreLoadScreen/Private/PreLoadScreenManager.cpp b/Engine/Source/Runtime/PreLoadScreen/Private/PreLoadScreenManager.cpp
--- a/Engine/Source/Runtime/PreLoadScreen/Private/PreLoadScreenManager.cpp
+++ b/Engine/Source/Runtime/PreLoadScreen/Private/PreLoadScreenManager.cpp
@@ -712,10 +712,6 @@ bool FPreLoadScreenManager::ArePreLoadScreensEnabled()
 	bEnabled &= !FParse::Param(FCommandLine::Get(), TEXT("NoLoadingScreen"));
 #endif
 
-#if PLATFORM_UNIX
-	bEnabled = false;
-#endif
-
 	return bEnabled;
 }
 
