Break out of EarlyPlayFrameTick loop if engine exit has been requested

Should fix the Ok button not doing anything for the Update Available dialog.

diff --git a/Engine/Source/Runtime/PreLoadScreen/Private/PreLoadScreenManager.cpp b/Engine/Source/Runtime/PreLoadScreen/Private/PreLoadScreenManager.cpp
--- a/Engine/Source/Runtime/PreLoadScreen/Private/PreLoadScreenManager.cpp
+++ b/Engine/Source/Runtime/PreLoadScreen/Private/PreLoadScreenManager.cpp
@@ -270,7 +270,7 @@ void FPreLoadScreenManager::HandleEarlyStartupPlay()
 				SCOPED_BOOT_TIMING("FPreLoadScreenManager::EarlyPlayFrameTick()");
 
 				//We run this PreLoadScreen until its finished or we lose the MainWindow as EarlyPreLoadPlay is synchronous
-				while (!PreLoadScreen->IsDone())
+				while (!PreLoadScreen->IsDone() && !IsEngineExitRequested())
 				{
 					EarlyPlayFrameTick();
 				}
