[PreLoadScreenBase] Allow ',' character in LocalizedText.
#rnx
#Synced-CL 52372219

diff --git a/Engine/Source/Runtime/PreLoadScreen/Private/PreLoadSettingsContainer.cpp b/Engine/Source/Runtime/PreLoadScreen/Private/PreLoadSettingsContainer.cpp
--- a/Engine/Source/Runtime/PreLoadScreen/Private/PreLoadSettingsContainer.cpp
+++ b/Engine/Source/Runtime/PreLoadScreen/Private/PreLoadSettingsContainer.cpp
@@ -383,7 +383,7 @@ bool FPreLoadSettingsContainerBase::BuildSystemFontFile()
 
 bool FPreLoadSettingsContainerBase::IsValidLocalizedTextConfigString(TArray<FString>& SplitConfigEntry)
 {
-    return (SplitConfigEntry.Num() == 4);
+    return (SplitConfigEntry.Num() >= 4);
 }
 
 void FPreLoadSettingsContainerBase::ParseLocalizedTextConfigString(const FString& ConfigEntry)
@@ -410,8 +410,14 @@ void FPreLoadSettingsContainerBase::ParseLocalizedTextConfigString(const FString
         LocIdentifier.RemoveFromStart("\"");
         LocIdentifier.RemoveFromEnd("\"");
 
-        //LocalizedTextComponents[3] is the default text for the FText
-        FString LocInitialValue = LocalizedTextComponents[3];
+    	//LocalizedTextComponents[3;End] is the default text for the FText
+		TArray<FString> LocInitialValueArray;
+    	for (int32 i = 3; i < LocalizedTextComponents.Num(); ++i)
+    	{
+    		LocInitialValueArray.Add(LocalizedTextComponents[i]);
+    	}
+    	//Joining [3;End] with ',' in case the original FText contains ',' and has been split in LocalizedTextComponents.
+    	FString LocInitialValue = FString::Join(LocInitialValueArray, TEXT(","));
         LocInitialValue.TrimStartAndEndInline();
         LocInitialValue.RemoveFromStart("\"");
         LocInitialValue.RemoveFromEnd(")"); //remove these separately so that if the file is missing one ) or the " is out of order it still works
